Move Modular Vehicle Plugin files out of restricted folders

diff --git a/Engine/Plugins/Experimental/ChaosModularVehicle/Source/ChaosModularVehicle/Public/ChaosModularVehicle/ModularSimCollection.h b/Engine/Plugins/Experimental/ChaosModularVehicle/Source/ChaosModularVehicle/Public/ChaosModularVehicle/ModularSimCollection.h
new file mode 100644
--- /dev/null
+++ b/Engine/Plugins/Experimental/ChaosModularVehicle/Source/ChaosModularVehicle/Public/ChaosModularVehicle/ModularSimCollection.h
@@ -0,0 +1,55 @@
+// Copyright Epic Games, Inc. All Rights Reserved.
+
+#pragma once
+
+#include "CoreMinimal.h"
+#include "GeometryCollection/GeometryCollection.h"
+
+namespace Chaos
+{
+	class FChaosArchive;
+}
+
+/**
+* FModularSimCollection (FTransformCollection)
+*/
+//class CHAOSMODULARVEHICLE_API FModularSimCollection : public FTransformCollection
+class CHAOSMODULARVEHICLE_API FModularSimCollection : public FGeometryCollection
+{
+public:
+	//typedef FTransformCollection Super;
+	typedef FGeometryCollection Super;
+
+	FModularSimCollection();
+	FModularSimCollection(FModularSimCollection&) = delete;
+	FModularSimCollection& operator=(const FModularSimCollection&) = delete;
+	FModularSimCollection(FModularSimCollection&&) = default;
+	FModularSimCollection& operator=(FModularSimCollection&&) = default;
+
+	/**
+	 * Create a GeometryCollection from Vertex and Indices arrays
+	 */
+	static FModularSimCollection* NewModularSimulationCollection(const FTransformCollection& Base);
+	static FModularSimCollection* NewModularSimulationCollection();
+	static void Init(FModularSimCollection* Collection);
+
+	/*
+	* Index of simulation module associated with each transform node
+	*   FManagedArray<int32> SimModuleIndex = this->FindAttribute<int32>("SimModuleIndex",FGeometryCollection::TransformGroup);
+	*/
+	static const FName SimModuleIndexAttribute;
+	TManagedArray<int32> SimModuleIndex;
+
+	// TODO: need to do this conversion somewhere - putting here for now
+	void GenerateSimTree();
+
+protected:
+	void Construct();
+
+};
+
+FORCEINLINE Chaos::FChaosArchive& operator<<(Chaos::FChaosArchive& Ar, FModularSimCollection& Value)
+{
+	Value.Serialize(Ar);
+	return Ar;
+}
diff --git a/Engine/Plugins/Experimental/ChaosModularVehicle/Source/ChaosModularVehicleEditor/ChaosModularVehicleEditor.Build.cs b/Engine/Plugins/Experimental/ChaosModularVehicle/Source/ChaosModularVehicleEditor/ChaosModularVehicleEditor.Build.cs
new file mode 100644
--- /dev/null
+++ b/Engine/Plugins/Experimental/ChaosModularVehicle/Source/ChaosModularVehicleEditor/ChaosModularVehicleEditor.Build.cs
@@ -0,0 +1,49 @@
+// Copyright Epic Games, Inc. All Rights Reserved.
+
+namespace UnrealBuildTool.Rules
+{
+	public class ChaosModularVehicleEditor : ModuleRules
+	{
+		public ChaosModularVehicleEditor(ReadOnlyTargetRules Target) : base(Target)
+		{
+			SetupModulePhysicsSupport(Target);
+
+			PublicDependencyModuleNames.AddRange(
+				new string[]
+				{
+
+				"Core",
+				"CoreUObject",
+				"Engine",
+				"InputCore",
+				"Slate",
+				"SlateCore",
+				"EditorStyle",
+				"EditorFramework",
+				"UnrealEd",
+				"MainFrame",
+				"DesktopPlatform",
+				"PropertyEditor",
+				"RHI",
+				"RawMesh",
+				"AssetTools",
+				"AssetRegistry",
+				"SceneOutliner",
+				"ToolMenus",
+				"EditorSubsystem",
+				"PhysicsUtilities",
+				"Chaos",
+				"ChaosVehiclesCore",
+				"ChaosModularVehicleEngine",
+				"GeometryCollectionEngine",
+				"GeometryCollectionEditor",
+				"ProceduralMeshComponent",			
+				"SubobjectEditor",
+				"SubobjectDataInterface",
+				}
+			);
+
+			PrivateDefinitions.Add("CHAOS_INCLUDE_LEVEL_1=1");
+		}
+	}
+}
diff --git a/Engine/Plugins/Experimental/ChaosModularVehicle/Source/ChaosModularVehicleEditor/Private/ChaosModularVehicle/ActorFactoryModularVehicle.cpp b/Engine/Plugins/Experimental/ChaosModularVehicle/Source/ChaosModularVehicleEditor/Private/ChaosModularVehicle/ActorFactoryModularVehicle.cpp
new file mode 100644
--- /dev/null
+++ b/Engine/Plugins/Experimental/ChaosModularVehicle/Source/ChaosModularVehicleEditor/Private/ChaosModularVehicle/ActorFactoryModularVehicle.cpp
@@ -0,0 +1,62 @@
+// Copyright Epic Games, Inc. All Rights Reserved.
+
+#include "ChaosModularVehicle/ActorFactoryModularVehicle.h"
+
+#include "ChaosModularVehicle/ModularVehiclePawn.h"
+#include "ChaosModularVehicle/ModularVehicleAsset.h"
+#include "ChaosModularVehicle/ModularVehicleComponent.h"
+#include "ChaosModularVehicle/ModularVehicleAssetFactory.h"
+#include "Misc/FileHelper.h"
+#include "AssetRegistry/AssetData.h"
+#include "Editor.h"
+#include "EditorFramework/AssetImportData.h"
+
+#define LOCTEXT_NAMESPACE "ActorFactoryModularVehicle"
+
+DEFINE_LOG_CATEGORY_STATIC(LogChaosModularVehicleFactories, Log, All);
+
+/*-----------------------------------------------------------------------------
+UActorFactoryModularVehicle
+-----------------------------------------------------------------------------*/
+UActorFactoryModularVehicle::UActorFactoryModularVehicle(const FObjectInitializer& ObjectInitializer)
+	: Super(ObjectInitializer)
+{
+	DisplayName = LOCTEXT("ModularVehicleDisplayName", "ModularVehicle");
+	NewActorClass = AModularVehiclePawn::StaticClass();
+}
+
+bool UActorFactoryModularVehicle::CanCreateActorFrom(const FAssetData& AssetData, FText& OutErrorMsg)
+{
+	if (!AssetData.IsValid() || !AssetData.GetClass()->IsChildOf(UModularVehicleAsset::StaticClass()))
+	{
+		OutErrorMsg = NSLOCTEXT("CanCreateActor", "NoModularVehicleSpecified", "No ModularVehicle asset was specified.");
+		return false;
+	}
+
+	return true;
+}
+
+void UActorFactoryModularVehicle::PostSpawnActor(UObject* Asset, AActor* NewActor)
+{
+	Super::PostSpawnActor(Asset, NewActor);
+
+	UModularVehicleAsset* ModularVehicle = CastChecked<UModularVehicleAsset>(Asset);
+	AModularVehiclePawn* NewModularVehicleActor = CastChecked<AModularVehiclePawn>(NewActor);
+
+	// Term Component
+	NewModularVehicleActor->GetModularVehicleComponent()->UnregisterComponent();
+
+	// Init Component
+	NewModularVehicleActor->GetModularVehicleComponent()->RegisterComponent();
+}
+
+void UActorFactoryModularVehicle::PostCreateBlueprint(UObject* Asset, AActor* CDO)
+{
+	if (Asset != NULL && CDO != NULL)
+	{
+		UModularVehicleAsset* ModularVehicle = CastChecked<UModularVehicleAsset>(Asset);
+		AModularVehiclePawn* ModularVehicleActor = CastChecked<AModularVehiclePawn>(CDO);
+	}
+}
+
+#undef LOCTEXT_NAMESPACE
\ No newline at end of file
diff --git a/Engine/Plugins/Experimental/ChaosModularVehicle/Source/ChaosModularVehicleEditor/Private/ChaosModularVehicle/ActorFactoryModularVehicle.h b/Engine/Plugins/Experimental/ChaosModularVehicle/Source/ChaosModularVehicleEditor/Private/ChaosModularVehicle/ActorFactoryModularVehicle.h
new file mode 100644
--- /dev/null
+++ b/Engine/Plugins/Experimental/ChaosModularVehicle/Source/ChaosModularVehicleEditor/Private/ChaosModularVehicle/ActorFactoryModularVehicle.h
@@ -0,0 +1,27 @@
+// Copyright Epic Games, Inc. All Rights Reserved.
+
+#pragma once
+
+#include "CoreMinimal.h"
+#include "UObject/ObjectMacros.h"
+#include "ActorFactories/ActorFactory.h"
+
+#include "ActorFactoryModularVehicle.generated.h"
+
+class AActor;
+struct FAssetData;
+
+UCLASS(MinimalAPI, config=Editor)
+class UActorFactoryModularVehicle : public UActorFactory
+{
+	GENERATED_UCLASS_BODY()
+
+	//~ Begin UActorFactory Interface
+	virtual void PostSpawnActor( UObject* Asset, AActor* NewActor ) override;
+	virtual void PostCreateBlueprint( UObject* Asset, AActor* CDO ) override;
+	virtual bool CanCreateActorFrom( const FAssetData& AssetData, FText& OutErrorMsg ) override;
+	//~ End UActorFactory Interface
+};
+
+
+
diff --git a/Engine/Plugins/Experimental/ChaosModularVehicle/Source/ChaosModularVehicleEditor/Private/ChaosModularVehicle/AssetTypeActions_ModularVehicleAsset.cpp b/Engine/Plugins/Experimental/ChaosModularVehicle/Source/ChaosModularVehicleEditor/Private/ChaosModularVehicle/AssetTypeActions_ModularVehicleAsset.cpp
new file mode 100644
--- /dev/null
+++ b/Engine/Plugins/Experimental/ChaosModularVehicle/Source/ChaosModularVehicleEditor/Private/ChaosModularVehicle/AssetTypeActions_ModularVehicleAsset.cpp
@@ -0,0 +1,35 @@
+// Copyright Epic Games, Inc. All Rights Reserved.
+
+#include "ChaosModularVehicle/AssetTypeActions_ModularVehicleAsset.h"
+
+#include "ThumbnailRendering/SceneThumbnailInfo.h"
+#include "ChaosModularVehicle/ModularVehicleAsset.h"
+#include "ToolMenus.h"
+
+#define LOCTEXT_NAMESPACE "AssetTypeActions"
+
+UClass* FAssetTypeActions_ModularVehicleAsset::GetSupportedClass() const
+{
+	return UModularVehicleAsset::StaticClass();
+}
+
+UThumbnailInfo* FAssetTypeActions_ModularVehicleAsset::GetThumbnailInfo(UObject* Asset) const
+{
+	UModularVehicleAsset * ModularVehicleAsset = CastChecked<UModularVehicleAsset>(Asset);
+	UThumbnailInfo* ThumbnailInfo = ModularVehicleAsset->ThumbnailInfo;
+	if (ThumbnailInfo == NULL)
+	{
+		ThumbnailInfo = NewObject<USceneThumbnailInfo>(ModularVehicleAsset, NAME_None, RF_Transactional);
+		ModularVehicleAsset->ThumbnailInfo = ThumbnailInfo;
+	}
+
+	return ThumbnailInfo;
+}
+
+void FAssetTypeActions_ModularVehicleAsset::GetActions(const TArray<UObject*>& InObjects, FToolMenuSection& Section)
+{
+	FAssetTypeActions_Base::GetActions(InObjects, Section);
+}
+
+
+#undef LOCTEXT_NAMESPACE
diff --git a/Engine/Plugins/Experimental/ChaosModularVehicle/Source/ChaosModularVehicleEditor/Private/ChaosModularVehicle/ChaosModularVehicleCommands.cpp b/Engine/Plugins/Experimental/ChaosModularVehicle/Source/ChaosModularVehicleEditor/Private/ChaosModularVehicle/ChaosModularVehicleCommands.cpp
new file mode 100644
--- /dev/null
+++ b/Engine/Plugins/Experimental/ChaosModularVehicle/Source/ChaosModularVehicleEditor/Private/ChaosModularVehicle/ChaosModularVehicleCommands.cpp
@@ -0,0 +1,25 @@
+// Copyright Epic Games, Inc. All Rights Reserved.
+
+#include "ChaosModularVehicle/ChaosModularVehicleCommands.h"
+
+#include "ChaosModularVehicle/ModularVehicleAsset.h"
+#include "ChaosModularVehicle/ModularVehicleActor.h"
+#include "ChaosModularVehicle/ModularVehicleComponent.h"
+#include "Engine/Selection.h"
+#include "Engine/World.h"
+#include "EngineUtils.h"
+#include "Internationalization/Regex.h"
+
+#include "Logging/LogMacros.h"
+#include "Editor.h"
+#include "EditorSupportDelegates.h"
+#include "SceneOutlinerDelegates.h"
+
+
+DEFINE_LOG_CATEGORY_STATIC(UChaosModularVehicleCommandsLogging, NoLogging, All);
+
+
+
+void FChaosModularVehicleCommands::ImportFile(const FString& File)
+{
+}
diff --git a/Engine/Plugins/Experimental/ChaosModularVehicle/Source/ChaosModularVehicleEditor/Private/ChaosModularVehicle/ChaosModularVehicleEditorPlugin.cpp b/Engine/Plugins/Experimental/ChaosModularVehicle/Source/ChaosModularVehicleEditor/Private/ChaosModularVehicle/ChaosModularVehicleEditorPlugin.cpp
new file mode 100644
--- /dev/null
+++ b/Engine/Plugins/Experimental/ChaosModularVehicle/Source/ChaosModularVehicleEditor/Private/ChaosModularVehicle/ChaosModularVehicleEditorPlugin.cpp
@@ -0,0 +1,72 @@
+// Copyright Epic Games, Inc. All Rights Reserved.
+
+
+#include "ChaosModularVehicle/ChaosModularVehicleEditorPlugin.h"
+
+#include "AssetToolsModule.h"
+#include "ChaosModularVehicle/ModularVehicleComponent.h"
+#include "ChaosModularVehicle/ModularVehicleAssetConversion.h"
+#include "ChaosModularVehicle/ChaosModularVehicleCommands.h"
+#include "ChaosModularVehicle/AssetTypeActions_ModularVehicleAsset.h"
+#include "ChaosModularVehicle/ModularVehicleAssetThumbnailRenderer.h"
+#include "ChaosModularVehicle/SimModuleActor.h"
+
+#include "CoreMinimal.h"
+#include "LevelEditor.h"
+#include "Modules/ModuleManager.h"
+#include "PropertyEditorModule.h"
+#include "HAL/ConsoleManager.h"
+#include "Features/IModularFeatures.h"
+#include "Styling/SlateStyle.h"
+#include "Styling/SlateStyleRegistry.h"
+#include "Styling/CoreStyle.h"
+#include "ToolMenus.h"
+
+
+#define BOX_BRUSH(StyleSet, RelativePath, ...) FSlateBoxBrush(StyleSet->RootToContentDir(RelativePath, TEXT(".png")), __VA_ARGS__)
+#define DEFAULT_FONT(...) FCoreStyle::GetDefaultFontStyle(__VA_ARGS__)
+
+class FChaosModularVehicleEditorPlugin : public IChaosModularVehicleEditorPlugin
+{
+	virtual void StartupModule() override;
+	virtual void ShutdownModule() override;
+
+	//void OnPostWorldInitialization(UWorld* InWorld, const UWorld::InitializationValues);
+
+};
+
+IMPLEMENT_MODULE(FChaosModularVehicleEditorPlugin, ModularVehicleAssetEditor);
+
+
+void FChaosModularVehicleEditorPlugin::StartupModule()
+{
+	FAssetToolsModule& AssetToolsModule = FAssetToolsModule::GetModule();
+	IAssetTools& AssetTools = AssetToolsModule.Get();
+	ModularVehicleAssetActions = new FAssetTypeActions_ModularVehicleAsset();
+	AssetTools.RegisterAssetTypeActions(MakeShareable(ModularVehicleAssetActions));
+
+	if (GIsEditor && !IsRunningCommandlet())
+	{
+
+	}
+
+	UThumbnailManager::Get().RegisterCustomRenderer(UModularVehicleAsset::StaticClass(), UModularVehicleAssetThumbnailRenderer::StaticClass());
+
+	//OnPostWorldInitializationHandle = FWorldDelegates::OnPostWorldInitialization.AddRaw(this, &FChaosModularVehicleEditorPlugin::OnPostWorldInitialization);
+
+}
+
+
+void FChaosModularVehicleEditorPlugin::ShutdownModule()
+{
+	if (UObjectInitialized())
+	{	
+		UThumbnailManager::Get().UnregisterCustomRenderer(UModularVehicleAsset::StaticClass());
+
+		FAssetToolsModule& AssetToolsModule = FAssetToolsModule::GetModule();
+		IAssetTools& AssetTools = AssetToolsModule.Get();
+		AssetTools.UnregisterAssetTypeActions(ModularVehicleAssetActions->AsShared());
+
+		//FWorldDelegates::OnPostWorldInitialization.Remove(OnPostWorldInitializationHandle);
+	}
+}
diff --git a/Engine/Plugins/Experimental/ChaosModularVehicle/Source/ChaosModularVehicleEditor/Private/ChaosModularVehicle/ModularToolkitBlueprintLibrary.cpp b/Engine/Plugins/Experimental/ChaosModularVehicle/Source/ChaosModularVehicleEditor/Private/ChaosModularVehicle/ModularToolkitBlueprintLibrary.cpp
new file mode 100644
--- /dev/null
+++ b/Engine/Plugins/Experimental/ChaosModularVehicle/Source/ChaosModularVehicleEditor/Private/ChaosModularVehicle/ModularToolkitBlueprintLibrary.cpp
@@ -0,0 +1,647 @@
+// Copyright Epic Games, Inc. All Rights Reserved.
+
+#include "ChaosModularVehicle/ModularToolkitBlueprintLibrary.h"
+
+#include "Components/StaticMeshComponent.h"
+#include "Editor.h"
+#include "Engine/Selection.h"
+#include "AssetToolsModule.h"
+#include "AssetRegistry/AssetRegistryModule.h"
+#include "Layers/LayersSubsystem.h"
+#include "Subsystems/UnrealEditorSubsystem.h"
+#include "Subsystems/AssetEditorSubsystem.h"
+#include "Subsystems/ImportSubsystem.h"
+#include "Kismet2/KismetEditorUtilities.h"
+#include "ScopedTransaction.h"
+
+#include "GeometryCollection/GeometryCollection.h"
+#include "GeometryCollection/GeometryCollectionComponent.h"
+#include "GeometryCollection/GeometryCollectionActor.h"
+#include "GeometryCollection/GeometryCollectionClusteringUtility.h"
+#include "GeometryCollection/GeometryCollectionConversion.h"
+#include "GeometryCollection/GeometryCollectionUtility.h"
+#include "ChaosModularVehicle/ModularVehicleObject.h"
+#include "ChaosModularVehicle/ModularVehicleComponent.h"
+#include "ChaosModularVehicle/ModularVehiclePawn.h"
+#include "ChaosModularVehicle/ModularSimCollection.h"
+#include "ChaosModularVehicle/VehicleSimComponentsInclude.h"
+#include "SimModule/SimModulesInclude.h"
+#include "VehicleUtility.h"
+
+#if VEHICLE_DEBUGGING_ENABLED
+UE_DISABLE_OPTIMIZATION
+#endif
+
+DEFINE_LOG_CATEGORY(LogModularVehicleToolkit);
+
+UModularToolkitBlueprintLibrary::UModularToolkitBlueprintLibrary(const FObjectInitializer& ObjectInitializer)
+	: Super(ObjectInitializer)
+{
+}
+
+
+void UModularToolkitBlueprintLibrary::AddSingleRootNodeIfRequired(UGeometryCollection* GeometryCollectionObject)
+{
+	TSharedPtr<FGeometryCollection, ESPMode::ThreadSafe> GeometryCollectionPtr = GeometryCollectionObject->GetGeometryCollection();
+	if (FGeometryCollection* GeometryCollection = GeometryCollectionPtr.Get())
+	{
+		if (FGeometryCollectionClusteringUtility::ContainsMultipleRootBones(GeometryCollection))
+		{
+			FGeometryCollectionClusteringUtility::ClusterAllBonesUnderNewRoot(GeometryCollection);
+		}
+	}
+}
+
+void UModularToolkitBlueprintLibrary::AddAdditionalAttributesIfRequired(UGeometryCollection* GeometryCollectionObject)
+{
+	TSharedPtr<FGeometryCollection, ESPMode::ThreadSafe> GeometryCollectionPtr = GeometryCollectionObject->GetGeometryCollection();
+	if (FGeometryCollection* GeometryCollection = GeometryCollectionPtr.Get())
+	{
+		if (!GeometryCollection->HasAttribute("Level", FGeometryCollection::TransformGroup))
+		{
+			FGeometryCollectionClusteringUtility::UpdateHierarchyLevelOfChildren(GeometryCollection, -1);
+		}
+	}
+}
+
+void UModularToolkitBlueprintLibrary::FlattenHierarchy(UGeometryCollection* GeometryCollectionObject)
+{
+	TSharedPtr<FGeometryCollection, ESPMode::ThreadSafe> GeometryCollectionPtr = GeometryCollectionObject->GetGeometryCollection();
+	if (FGeometryCollection* GeometryCollection = GeometryCollectionPtr.Get())
+	{
+		// Make sure we have valid "Level", this is not a default/runtime attribute, it is only used in editor
+		AddAdditionalAttributesIfRequired(GeometryCollectionObject);
+
+		if (GeometryCollection->Transform.Num() > 1)
+		{
+			FGeometryCollectionClusteringUtility::CollapseLevelHierarchy(-1, GeometryCollection);
+		}
+	}
+}
+
+ULevel* UModularToolkitBlueprintLibrary::GetSelectedLevel()
+{
+	USelection* SelectedActors = GEditor->GetSelectedActors();
+	TArray<ULevel*> UniqueLevels;
+	for (FSelectionIterator Iter(*SelectedActors); Iter; ++Iter)
+	{
+		AActor* Actor = Cast<AActor>(*Iter);
+		if (Actor)
+		{
+			UniqueLevels.AddUnique(Actor->GetLevel());
+		}
+	}
+	check(UniqueLevels.Num() == 1);
+	return UniqueLevels[0];
+}
+
+AActor* UModularToolkitBlueprintLibrary::AddActor(ULevel* InLevel, UClass* Class)
+{
+	check(Class);
+
+	UWorld* World = InLevel->OwningWorld;
+	ULevel* DesiredLevel = InLevel;
+
+	// Transactionally add the actor.
+	AActor* Actor = NULL;
+	{
+		FScopedTransaction Transaction(NSLOCTEXT("UnrealEd", "AddActor", "Add Actor"));
+
+		FActorSpawnParameters SpawnInfo;
+		SpawnInfo.OverrideLevel = DesiredLevel;
+		SpawnInfo.SpawnCollisionHandlingOverride = ESpawnActorCollisionHandlingMethod::AlwaysSpawn;
+		SpawnInfo.ObjectFlags = RF_Transactional;
+		const auto Location = FVector(0);
+		const auto Rotation = FTransform(FVector(0)).GetRotation().Rotator();
+		Actor = World->SpawnActor(Class, &Location, &Rotation, SpawnInfo);
+
+		check(Actor);
+		Actor->InvalidateLightingCache();
+		Actor->PostEditMove(true);
+	}
+
+	// If this actor is part of any layers (set in its default properties), add them into the visible layers list.
+	ULayersSubsystem* Layers = GEditor->GetEditorSubsystem<ULayersSubsystem>();
+	Layers->SetLayersVisibility(Actor->Layers, true);
+
+	// Clean up.
+	Actor->MarkPackageDirty();
+	ULevel::LevelDirtiedEvent.Broadcast();
+
+	return Actor;
+}
+
+
+void UModularToolkitBlueprintLibrary::AddActorToExistingModularVehicle(AModularVehiclePawn* ExistingActorInOut, const AActor* Actor /*const TArray<AActor*>& Actors*/)
+{
+	//	ensure(Actors.Num() > 0);
+	const AActor* FirstActor = ExistingActorInOut;
+	const FString& Name = FirstActor->GetActorLabel();
+	const FVector FirstActorLocation(FirstActor->GetActorLocation());
+
+	FGeometryCollectionEdit GeometryCollectionEdit = ExistingActorInOut->GetModularVehicleComponent()->EditRestCollection(GeometryCollection::EEditUpdate::RestPhysicsDynamic);
+	UGeometryCollection* UpdatedGeometryCollection = GeometryCollectionEdit.GetRestCollection();
+	check(UpdatedGeometryCollection);
+
+	//for (AActor* Actor : Actors)
+	{
+		const FTransform ActorTransform(Actor->GetTransform());
+		const FVector ActorOffset(Actor->GetActorLocation() - FirstActor->GetActorLocation());
+
+		TArray<UStaticMeshComponent*> StaticMeshComponents;
+		Actor->GetComponents<UStaticMeshComponent>(StaticMeshComponents, true);
+
+		// START....
+		TArray<UActorComponent*> SimComponents;
+		Actor->GetComponents(UVehicleSimBaseComponent::StaticClass(), SimComponents);
+
+		UVehicleSimBaseComponent* SimBaseComponent = Cast<UVehicleSimBaseComponent>(SimComponents[0]);
+		bool IsWheel = SimBaseComponent->IsA(UVehicleSimWheelComponent::StaticClass());
+
+		// ....END
+
+		for (int32 ii = 0, ni = StaticMeshComponents.Num(); ii < ni; ++ii)
+		{
+			check(StaticMeshComponents.Num() == 1); // currently only supports one static mesh - because sim indices to transforms get all mixed up otherwise
+
+			// We're partial to static mesh components, here
+			UStaticMeshComponent* StaticMeshComponent = StaticMeshComponents[ii];
+			if (StaticMeshComponent != nullptr)
+			{
+				UStaticMesh* ComponentStaticMesh = StaticMeshComponent->GetStaticMesh();
+				if (ComponentStaticMesh != nullptr)
+				{
+					// If any of the static meshes have Nanite enabled, also enable on the new geometry collection asset for convenience.
+					UpdatedGeometryCollection->EnableNanite |= ComponentStaticMesh->IsNaniteEnabled();
+
+					FTransform ComponentTransform(StaticMeshComponent->GetComponentTransform());
+					ComponentTransform.SetTranslation((ComponentTransform.GetTranslation() - ActorTransform.GetTranslation()) + ActorOffset);
+
+					// Record the contributing source on the asset.
+					FSoftObjectPath SourceSoftObjectPath(ComponentStaticMesh);
+					decltype(FGeometryCollectionSource::SourceMaterial) SourceMaterials(StaticMeshComponent->GetMaterials());
+					UpdatedGeometryCollection->GeometrySource.Emplace(SourceSoftObjectPath, ComponentTransform, SourceMaterials);
+
+					//if (IsWheel)
+					//{
+					//	UpdatedGeometryCollection->GetGeometryCollection()->SimulationType[] = FGeometryCollection::ESimulationTypes::FST_None;
+					//}
+					FGeometryCollectionConversion::AppendStaticMesh(ComponentStaticMesh, SourceMaterials, ComponentTransform, UpdatedGeometryCollection, true);
+				}
+			}
+		}
+
+		TArray<UModularVehicleComponent*> ModularVehicleComponents;
+		Actor->GetComponents<UModularVehicleComponent>(ModularVehicleComponents, true);
+		for (int32 ii = 0, ni = ModularVehicleComponents.Num(); ii < ni; ++ii)
+		{
+			UModularVehicleComponent* ModularVehicleComponent = ModularVehicleComponents[ii];
+			if (ModularVehicleComponent != nullptr)
+			{
+				const UGeometryCollection* RestCollection = ModularVehicleComponent->GetRestCollection();
+				if (RestCollection != nullptr)
+				{
+					// If any of the static meshes have Nanite enabled, also enable on the new geometry collection asset for convenience.
+					UpdatedGeometryCollection->EnableNanite |= RestCollection->EnableNanite;
+				}
+
+				FTransform ComponentTransform(ModularVehicleComponent->GetComponentTransform());
+				ComponentTransform.SetTranslation((ComponentTransform.GetTranslation() - ActorTransform.GetTranslation()) + ActorOffset);
+
+				// Record the contributing source on the asset.
+				FSoftObjectPath SourceSoftObjectPath(RestCollection);
+
+				// We're not interested in recording the final material of the collection since it's inevitably the Selection material.
+				int32 NumMaterials = ModularVehicleComponent->GetNumMaterials() - 1;
+				TArray<TObjectPtr<UMaterialInterface>> SourceMaterials;
+				SourceMaterials.SetNum(NumMaterials);
+				for (int32 MaterialIndex = 0; MaterialIndex < NumMaterials; ++MaterialIndex)
+				{
+					SourceMaterials[MaterialIndex] = ModularVehicleComponent->GetMaterial(MaterialIndex);
+				}
+				UpdatedGeometryCollection->GeometrySource.Emplace(SourceSoftObjectPath, ComponentTransform, SourceMaterials);
+
+				FGeometryCollectionConversion::AppendGeometryCollection(RestCollection, ModularVehicleComponent, ComponentTransform, UpdatedGeometryCollection, false);
+
+			}
+		}
+
+	}
+
+	UpdatedGeometryCollection->InitializeMaterials();
+
+	// TODO: best if we add the actors in a flat manner in the first place, then we can skip this step
+	// different behavior from fracture tools as it creates a new GC whereas we are adding to an existing one
+	AddSingleRootNodeIfRequired(UpdatedGeometryCollection);
+	FlattenHierarchy(UpdatedGeometryCollection);
+
+	if (UpdatedGeometryCollection->EnableNanite)
+	{
+		UpdatedGeometryCollection->InvalidateCollection();
+		UpdatedGeometryCollection->RebuildRenderData();
+	}
+
+	ExistingActorInOut->GetModularVehicleComponent()->MarkRenderStateDirty();
+
+	// Add and initialize guids
+	::GeometryCollection::GenerateTemporaryGuids(UpdatedGeometryCollection->GetGeometryCollection().Get(), 0, true);
+}
+
+void UModularToolkitBlueprintLibrary::AddActorsToExistingModularVehicle(AModularVehiclePawn* ExistingActorInOut, const TArray<AActor*>& Actors)
+{
+	ensure(Actors.Num() > 0);
+	const AActor* FirstActor = Actors[0];
+	const FString& Name = FirstActor->GetActorLabel();
+	const FVector FirstActorLocation(FirstActor->GetActorLocation());
+
+	FGeometryCollectionEdit GeometryCollectionEdit = ExistingActorInOut->GetModularVehicleComponent()->EditRestCollection(GeometryCollection::EEditUpdate::RestPhysicsDynamic);
+	UGeometryCollection* UpdatedGeometryCollection = GeometryCollectionEdit.GetRestCollection();
+	check(UpdatedGeometryCollection);
+
+	for (AActor* Actor : Actors)
+	{
+		UE_LOG(LogModularVehicleToolkit, Log, TEXT("ModularVehicleConvert Adding Actor %s"), *Actor->GetName());
+
+		const FTransform ActorTransform(Actor->GetTransform());
+		const FVector ActorOffset(Actor->GetActorLocation() - FirstActor->GetActorLocation());
+
+		TArray<UStaticMeshComponent*> StaticMeshComponents;
+		Actor->GetComponents<UStaticMeshComponent>(StaticMeshComponents, true);
+
+		for (int32 ii = 0, ni = StaticMeshComponents.Num(); ii < ni; ++ii)
+		{
+			check(StaticMeshComponents.Num() == 1); // currently only supports one static mesh - because sim indices to transforms get all mixed up otherwise
+
+			UStaticMeshComponent* StaticMeshComponent = StaticMeshComponents[ii];
+			if (StaticMeshComponent != nullptr)
+			{
+				UStaticMesh* ComponentStaticMesh = StaticMeshComponent->GetStaticMesh();
+				if (ComponentStaticMesh != nullptr)
+				{
+					// If any of the static meshes have Nanite enabled, also enable on the new geometry collection asset for convenience.
+					UpdatedGeometryCollection->EnableNanite |= ComponentStaticMesh->IsNaniteEnabled();
+
+					FTransform ComponentTransform(StaticMeshComponent->GetComponentTransform());
+					ComponentTransform.SetTranslation((ComponentTransform.GetTranslation() - ActorTransform.GetTranslation()) + ActorOffset);
+
+					// Record the contributing source on the asset.
+					FSoftObjectPath SourceSoftObjectPath(ComponentStaticMesh);
+					decltype(FGeometryCollectionSource::SourceMaterial) SourceMaterials(StaticMeshComponent->GetMaterials());
+					UpdatedGeometryCollection->GeometrySource.Emplace(SourceSoftObjectPath, ComponentTransform, SourceMaterials);
+
+					FGeometryCollectionConversion::AppendStaticMesh(ComponentStaticMesh, SourceMaterials, ComponentTransform, UpdatedGeometryCollection, true);
+
+					int LastIndex = UpdatedGeometryCollection->GetGeometryCollection()->Transform.Num() - 1;
+					UE_LOG(LogModularVehicleToolkit, Log, TEXT("FGeometryCollectionConversion:: SM_Name %s, added transform @ Index %d, BoneName %s")
+						, *ComponentStaticMesh->GetName()
+						, LastIndex
+						, *UpdatedGeometryCollection->GetGeometryCollection()->BoneName[LastIndex]);
+
+				}			
+			}
+		}
+
+		TArray<UModularVehicleComponent*> ModularVehicleComponents;
+		Actor->GetComponents<UModularVehicleComponent>(ModularVehicleComponents, true);
+		for (int32 ii = 0, ni = ModularVehicleComponents.Num(); ii < ni; ++ii)
+		{
+			UModularVehicleComponent* ModularVehicleComponent = ModularVehicleComponents[ii];
+			if (ModularVehicleComponent != nullptr)
+			{
+				const UGeometryCollection* RestCollection = ModularVehicleComponent->GetRestCollection();
+				if (RestCollection != nullptr)
+				{
+					// If any of the static meshes have Nanite enabled, also enable on the new geometry collection asset for convenience.
+					UpdatedGeometryCollection->EnableNanite |= RestCollection->EnableNanite;
+				}
+
+				// this isn't right - only works when not rotated or scaled
+				FTransform ComponentTransform(ModularVehicleComponent->GetComponentTransform());
+				ComponentTransform.SetTranslation((ComponentTransform.GetTranslation() - ActorTransform.GetTranslation()) + ActorOffset);
+
+				// Record the contributing source on the asset.
+				FSoftObjectPath SourceSoftObjectPath(RestCollection);
+
+				// We're not interested in recording the final material of the collection since it's inevitably the Selection material.
+				int32 NumMaterials = ModularVehicleComponent->GetNumMaterials() - 1;
+				TArray<TObjectPtr<UMaterialInterface>> SourceMaterials;
+				SourceMaterials.SetNum(NumMaterials);
+				for (int32 MaterialIndex = 0; MaterialIndex < NumMaterials; ++MaterialIndex)
+				{
+					SourceMaterials[MaterialIndex] = ModularVehicleComponent->GetMaterial(MaterialIndex);
+				}
+				UpdatedGeometryCollection->GeometrySource.Emplace(SourceSoftObjectPath, ComponentTransform, SourceMaterials);
+
+				FGeometryCollectionConversion::AppendGeometryCollection(RestCollection, ModularVehicleComponent, ComponentTransform, UpdatedGeometryCollection, false);
+			}
+		}
+
+	}
+
+	UpdatedGeometryCollection->InitializeMaterials();
+
+	// TODO: best if we add the actors in a flat manner in the first place if this is possible, then we can skip this step
+	// different behavior from fracture tools as it creates a new GC whereas we are adding to an existing one
+	AddSingleRootNodeIfRequired(UpdatedGeometryCollection);
+	FlattenHierarchy(UpdatedGeometryCollection);
+
+	if (UpdatedGeometryCollection->EnableNanite)
+	{
+		UpdatedGeometryCollection->InvalidateCollection();
+		UpdatedGeometryCollection->RebuildRenderData();
+	}
+
+	ExistingActorInOut->GetModularVehicleComponent()->MarkRenderStateDirty();
+
+	// Add and initialize guids
+	::GeometryCollection::GenerateTemporaryGuids(UpdatedGeometryCollection->GetGeometryCollection().Get(), 0, true);
+
+}
+
+
+AModularVehiclePawn* UModularToolkitBlueprintLibrary::ConvertActorsToModularVehicle(const FString& InAssetPath, const TArray<AActor*>& Actors)
+{
+	ensure(Actors.Num() > 0);
+
+	// create a new GC based actor
+	AModularVehiclePawn* NewActor = CreateNewModularVehicleActor(InAssetPath, FTransform());
+
+	// now append geometry to this new actor
+	AddActorsToExistingModularVehicle(NewActor, Actors);
+
+	// fix up the simulation hierarchy
+	int Index = 0;
+	for (AActor* Actor : Actors)
+	{
+		int32 TreeIndex = UModularToolkitBlueprintLibrary::UpdateSimulationTree(NewActor, EModuleOperation::Add, Index++, Actor);
+	}
+
+	// Automatically introduce expected vehicle simulation hiearchy - a future vehicle editor would do this for us
+	if (UModularVehicleComponent* ModularVehicleComponent = NewActor->GetModularVehicleComponent())
+	{
+		AttachNearest(ModularVehicleComponent, UVehicleSimSuspensionComponent::StaticClass(), UVehicleSimWheelComponent::StaticClass());
+
+		if (!Attach(ModularVehicleComponent, UVehicleSimWheelComponent::StaticClass(), UVehicleSimTransmissionComponent::StaticClass()))
+		{
+			if (!Attach(ModularVehicleComponent, UVehicleSimWheelComponent::StaticClass(), UVehicleSimClutchComponent::StaticClass()))
+			{
+				Attach(ModularVehicleComponent, UVehicleSimWheelComponent::StaticClass(), UVehicleSimEngineComponent::StaticClass());
+			}
+		}
+
+		Attach(ModularVehicleComponent, UVehicleSimTransmissionComponent::StaticClass(), UVehicleSimClutchComponent::StaticClass());
+		Attach(ModularVehicleComponent, UVehicleSimClutchComponent::StaticClass(), UVehicleSimEngineComponent::StaticClass());
+	}
+
+
+	// Create a BP Actor asset using this new Actor
+	FKismetEditorUtilities::FCreateBlueprintFromActorParams Params;
+	Params.bReplaceActor = true;
+	Params.ParentClassOverride = AModularVehiclePawn::StaticClass();
+	UBlueprint* Blueprint = FKismetEditorUtilities::CreateBlueprintFromActor(InAssetPath + FString("_BP"), NewActor, Params);
+
+	// Select the newly created blueprint in the content browser, but don't activate the browser
+	TArray<UObject*> Objects;
+	Objects.Add(Blueprint);
+	GEditor->SyncBrowserToObjects(Objects, false);
+
+	return NewActor;
+}
+
+class AModularVehiclePawn* UModularToolkitBlueprintLibrary::CreateNewModularVehicleActor(const FString& InAssetPath, const FTransform& Transform)
+{
+	FString UniquePackageName = InAssetPath + FString("_GC");
+	FString UniqueAssetName = FPackageName::GetLongPackageAssetName(InAssetPath + FString("_GC"));
+
+	FAssetToolsModule& AssetToolsModule = FModuleManager::LoadModuleChecked<FAssetToolsModule>("AssetTools");
+	AssetToolsModule.Get().CreateUniqueAssetName(UniquePackageName, TEXT(""), UniquePackageName, UniqueAssetName);
+
+	UPackage* Package = CreatePackage(*UniquePackageName);
+	UModularVehicle* InModularVehicle = static_cast<UModularVehicle*>(NewObject<UModularVehicle>(Package, UModularVehicle::StaticClass(), FName(*UniqueAssetName), RF_Transactional | RF_Public | RF_Standalone));
+	if (!InModularVehicle->SizeSpecificData.Num()) InModularVehicle->SizeSpecificData.Add(FGeometryCollectionSizeSpecificData());
+
+	// Create the new actor
+	AModularVehiclePawn* NewActor = Cast<AModularVehiclePawn>(AddActor(GetSelectedLevel(), AModularVehiclePawn::StaticClass()));
+
+	// Set the Geometry Collection asset in the new actor
+	NewActor->GetModularVehicleComponent()->SetRestCollection(InModularVehicle);
+
+	// copy transform of original static mesh actor to this new actor
+	NewActor->SetActorLabel(UniqueAssetName);
+	NewActor->SetActorTransform(Transform);
+
+	// Mark relevant stuff dirty
+	FAssetRegistryModule::AssetCreated(InModularVehicle);
+	InModularVehicle->MarkPackageDirty();
+	Package->SetDirtyFlag(true);
+
+	return NewActor;
+}
+
+
+class AModularVehiclePawn* UModularToolkitBlueprintLibrary::SetupRestCollectionOnActor(const FString& InAssetPath, AModularVehiclePawn* ModularVehicleActor, const FTransform& Transform)
+{
+	check(ModularVehicleActor);
+
+	FString UniquePackageName = InAssetPath;
+	FString UniqueAssetName = FPackageName::GetLongPackageAssetName(InAssetPath);
+
+	FAssetToolsModule& AssetToolsModule = FModuleManager::LoadModuleChecked<FAssetToolsModule>("AssetTools");
+	AssetToolsModule.Get().CreateUniqueAssetName(UniquePackageName, TEXT(""), UniquePackageName, UniqueAssetName);
+
+	UPackage* Package = CreatePackage(*UniquePackageName);
+
+	UModularVehicle* ModularVehicleAsset = static_cast<UModularVehicle*>(NewObject<UModularVehicle>(Package, UModularVehicle::StaticClass(), FName(*UniqueAssetName), RF_Transactional | RF_Public | RF_Standalone));
+	if (!ModularVehicleAsset->SizeSpecificData.Num()) ModularVehicleAsset->SizeSpecificData.Add(FGeometryCollectionSizeSpecificData());
+
+	check(ModularVehicleActor->GetModularVehicleComponent());
+
+	// Set the Geometry Collection asset in the new actor
+	ModularVehicleActor->GetModularVehicleComponent()->SetRestCollection(ModularVehicleAsset);
+
+	// copy transform of original static mesh actor to this new actor
+	ModularVehicleActor->SetActorLabel(UniqueAssetName);
+	ModularVehicleActor->SetActorTransform(Transform);
+
+	// Mark relevant stuff dirty
+	FAssetRegistryModule::AssetCreated(ModularVehicleAsset);
+	ModularVehicleAsset->MarkPackageDirty();
+	Package->SetDirtyFlag(true);
+
+	return ModularVehicleActor;
+}
+
+
+int32 UModularToolkitBlueprintLibrary::UpdateSimulationTree(AModularVehiclePawn* ExistingVehicle, EModuleOperation Operation, int32 Index, const AActor* Actor)
+{
+	check(ExistingVehicle);
+	if (UModularVehicleComponent* ModularVehicleComponent = ExistingVehicle->GetModularVehicleComponent())
+	{
+		if (const UGeometryCollection* Rest = ModularVehicleComponent->GetRestCollection())
+		{
+			USceneComponent* BaseComponent = ModularVehicleComponent;
+
+			TArray<UActorComponent*> Components;
+			Actor->GetComponents(UVehicleSimBaseComponent::StaticClass(), Components);
+
+			for (int32 ii = 0, ni = Components.Num(); ii < ni; ++ii)
+			{
+				UVehicleSimBaseComponent* SimBaseComponent = Cast<UVehicleSimBaseComponent>(Components[ii]);
+
+				UE_LOG(LogModularVehicleToolkit, Log, TEXT("UpdateSimulationTree Actor %s, Comp %s, Index %d?"), *Actor->GetName(), *SimBaseComponent->GetName(), Index);
+
+				UVehicleSimBaseComponent* DuplicateComponent = DuplicateObject(SimBaseComponent, ExistingVehicle);
+
+				if (const TSharedPtr<FGeometryCollection, ESPMode::ThreadSafe> GC = Rest->GetGeometryCollection())
+				{
+					UE_LOG(LogModularVehicleToolkit, Log, TEXT("Index %d < GC->Transform.Num() %d"), Index, GC->Transform.Num());
+					if (Index >= GC->Transform.Num())
+					{
+						UE_LOG(LogModularVehicleToolkit, Warning, TEXT("Modular Vehicle Index out of range"));
+						return Chaos::ISimulationModuleBase::INVALID_IDX;
+					}
+
+					// Will ignore all the chassis parts as they don't have a function anyway
+					if (SimBaseComponent->GetClass() != UVehicleSimChassisComponent::StaticClass())
+					{
+						FTransform NewComponentTransform = FTransform(GC->Transform[Index]);
+						DuplicateComponent->TransformIndex = Index;
+						DuplicateComponent->SetWorldTransform(NewComponentTransform);
+
+						// exists to provide data setup information for the vehicle simulation, component itself does no work/is not visible
+						DuplicateComponent->SetIsReplicated(false);
+						DuplicateComponent->SetActive(false);
+						DuplicateComponent->SetComponentTickEnabled(false);
+						DuplicateComponent->SetAutoActivate(false);
+						DuplicateComponent->SetHiddenInGame(true);
+
+						// taking some clues as what to do here from USubobjectDataSubsystem::AddNewSubobject
+						DuplicateComponent->AttachToComponent(BaseComponent, FAttachmentTransformRules::KeepWorldTransform);
+						ExistingVehicle->AddInstanceComponent(DuplicateComponent);
+						DuplicateComponent->OnComponentCreated();
+						DuplicateComponent->RegisterComponent();
+					}
+
+				}
+			}
+
+			ExistingVehicle->RerunConstructionScripts();
+		}
+	}
+
+	return Chaos::ISimulationModuleBase::INVALID_IDX;
+}
+
+bool UModularToolkitBlueprintLibrary::Attach(UModularVehicleComponent* ModularVehicleComponent, UClass* AttachClass, UClass* ToClass)
+{
+	bool bMatchFound = false;
+	TArray<TObjectPtr<USceneComponent>> Children = ModularVehicleComponent->GetAttachChildren(); // deliberate copy
+
+	TObjectPtr<USceneComponent> To = nullptr;
+	for (auto& Child : Children)
+	{
+		if (Child.GetClass() == ToClass)
+		{
+			To = Child;
+		}
+	}
+
+	if (To != nullptr)
+	{
+		for (auto& Child : Children)
+		{
+			if (Child.GetClass() == AttachClass)
+			{
+				bMatchFound = true;
+				Child->AttachToComponent(To, FAttachmentTransformRules::KeepWorldTransform);
+			}
+		}
+	}
+
+	return bMatchFound;
+}
+
+bool UModularToolkitBlueprintLibrary::AttachNearest(UModularVehicleComponent* ModularVehicleComponent, UClass* AttachClass, UClass* ToClass)
+{
+	bool bMatchFound = false;
+	TArray<TObjectPtr<USceneComponent>> Children = ModularVehicleComponent->GetAttachChildren(); // deliberate copy
+
+	for (auto& Child : Children)
+	{
+		if (Child.GetClass() == AttachClass)
+		{
+			TObjectPtr<USceneComponent> Me = Child;
+			TObjectPtr<USceneComponent> To = nullptr;
+			float ClosestDistSqr = MAX_FLT;
+
+			for (auto& ChildAlt : Children)
+			{
+				if (ChildAlt.GetClass() == ToClass)
+				{
+					bMatchFound = true;
+					float DistanceSqr = (Me->GetRelativeLocation() - ChildAlt->GetRelativeLocation()).SizeSquared();
+					if (DistanceSqr < ClosestDistSqr)
+					{
+						ClosestDistSqr = DistanceSqr;
+						To = ChildAlt;
+					}
+				}
+			}
+
+			if (Me && To)
+			{
+				Me->AttachToComponent(To, FAttachmentTransformRules::KeepWorldTransform);
+			}
+		}
+	}
+
+	return bMatchFound;
+}
+
+
+int32 UModularToolkitBlueprintLibrary::AddChild(AModularVehiclePawn* ExistingVehicleActor, int32 ParentIndex, const AActor* Actor)
+{
+	// if GC Rest collection does not exist create it
+	if (ExistingVehicleActor->GetModularVehicleComponent()->GetRestCollection() == nullptr)
+	{
+		FString OriginalName = ExistingVehicleActor->GetName();
+		const FString InAssetPath = FString("/Game/Developers/billhenderson/" + OriginalName); // #TODO: fix
+		SetupRestCollectionOnActor(InAssetPath, ExistingVehicleActor, FTransform());
+	}
+
+	// append geometry to this actor - flat hierarchy of nodes in Geometry Collection
+	UModularToolkitBlueprintLibrary::AddActorToExistingModularVehicle(ExistingVehicleActor, Actor);
+
+	// fix up the simulation hierarchy
+	int32 TreeIndex = UModularToolkitBlueprintLibrary::UpdateSimulationTree(ExistingVehicleActor, EModuleOperation::Add, ParentIndex, Actor);
+
+	return TreeIndex;
+}
+
+int32 UModularToolkitBlueprintLibrary::InsertAbove(AModularVehiclePawn* ExistingVehicleActor, int32 ChildIndex, const AActor* Actor)
+{
+	// if GC Rest collection does not exist create it
+	//if (ExistingVehicleActor->GetModularVehicleComponent()->GetRestCollection() == nullptr)
+	//{
+	//	FTransform Transform;
+	//	SetupRestCollectionOnActor(ExistingVehicleActor, Transform);
+	//}
+
+	// append geometry to this actor - flat hierarchy of nodes in Geometry Collection
+	UModularToolkitBlueprintLibrary::AddActorToExistingModularVehicle(ExistingVehicleActor, Actor);
+
+	// fix up the simulation hierarchy
+	int32 TreeIndex = UModularToolkitBlueprintLibrary::UpdateSimulationTree(ExistingVehicleActor, EModuleOperation::Insert, ChildIndex, Actor);
+
+	return TreeIndex;
+}
+
+int32 UModularToolkitBlueprintLibrary::RemoveAt(AModularVehiclePawn* ExistingVehicleActor, int Index)
+{
+	check(false); // TODO: implement
+	return -1;
+}
+
+#if VEHICLE_DEBUGGING_ENABLED
+UE_ENABLE_OPTIMIZATION
+#endif
diff --git a/Engine/Plugins/Experimental/ChaosModularVehicle/Source/ChaosModularVehicleEditor/Private/ChaosModularVehicle/ModularVehicleAssetConversion.cpp b/Engine/Plugins/Experimental/ChaosModularVehicle/Source/ChaosModularVehicleEditor/Private/ChaosModularVehicle/ModularVehicleAssetConversion.cpp
new file mode 100644
--- /dev/null
+++ b/Engine/Plugins/Experimental/ChaosModularVehicle/Source/ChaosModularVehicleEditor/Private/ChaosModularVehicle/ModularVehicleAssetConversion.cpp
@@ -0,0 +1,17 @@
+// Copyright Epic Games, Inc. All Rights Reserved.
+
+#include "ChaosModularVehicle/ModularVehicleAssetConversion.h"
+
+#include "ChaosModularVehicle/ModularVehicleAsset.h"
+#include "ChaosModularVehicle/ModularVehicleActor.h"
+#include "ChaosModularVehicle/ModularVehicleAssetFactory.h"
+#include "ChaosModularVehicle/ModularVehicleComponent.h"
+#include "Logging/LogMacros.h"
+
+
+DEFINE_LOG_CATEGORY_STATIC(UModularVehicleAssetConversionLogging, Log, All);
+
+
+void FModularVehicleAssetConversion::ImportFile(FString & FileName)
+{
+}
\ No newline at end of file
diff --git a/Engine/Plugins/Experimental/ChaosModularVehicle/Source/ChaosModularVehicleEditor/Private/ChaosModularVehicle/ModularVehicleAssetFactory.cpp b/Engine/Plugins/Experimental/ChaosModularVehicle/Source/ChaosModularVehicleEditor/Private/ChaosModularVehicle/ModularVehicleAssetFactory.cpp
new file mode 100644
--- /dev/null
+++ b/Engine/Plugins/Experimental/ChaosModularVehicle/Source/ChaosModularVehicleEditor/Private/ChaosModularVehicle/ModularVehicleAssetFactory.cpp
@@ -0,0 +1,67 @@
+// Copyright Epic Games, Inc. All Rights Reserved.
+
+#include "ChaosModularVehicle/ModularVehicleAssetFactory.h"
+
+#include "ChaosModularVehicle/ModularVehicleAsset.h"
+#include "ChaosModularVehicle/ModularVehicleComponent.h"
+#include "Editor.h"
+#include "Editor/EditorEngine.h"
+#include "Engine/Selection.h"
+
+
+#define LOCTEXT_NAMESPACE "ChaosModularVehicle"
+
+/////////////////////////////////////////////////////
+// ModularVehicleFactory
+
+UModularVehicleAssetFactory::UModularVehicleAssetFactory(const FObjectInitializer& ObjectInitializer)
+	: Super(ObjectInitializer)
+{
+	bCreateNew = true;
+	bEditAfterNew = true;
+	SupportedClass = UModularVehicleAsset::StaticClass();
+}
+
+UModularVehicleAsset* UModularVehicleAssetFactory::StaticFactoryCreateNew(UClass* Class, UObject* InParent, FName Name, EObjectFlags Flags, UObject* Context, FFeedbackContext* Warn)
+{
+	UModularVehicleAsset* NewModularVehicle = static_cast<UModularVehicleAsset*>(NewObject<UModularVehicleAsset>(InParent, Class, Name, Flags | RF_Transactional | RF_Public | RF_Standalone));
+	return NewModularVehicle;
+}
+
+UObject* UModularVehicleAssetFactory::FactoryCreateNew(UClass* Class, UObject* InParent, FName Name, EObjectFlags Flags, UObject* Context, FFeedbackContext* Warn)
+{
+	FTransform LastTransform = FTransform::Identity;
+	TArray< ModularVehicleTuple > ModularVehicleList;
+
+	TArray<FAssetData> SelectedAssets;
+	GEditor->GetContentBrowserSelections(SelectedAssets);
+
+	for (const FAssetData& AssetData : SelectedAssets)
+	{
+		if (AssetData.GetAsset()->IsA<UModularVehicleAsset>())
+		{
+			UModularVehicleComponent *DummyValue(NULL);
+			ModularVehicleList.Add(ModularVehicleTuple(static_cast<const UModularVehicleAsset *>(AssetData.GetAsset()), DummyValue, FTransform()));
+		}
+
+	}
+
+	UModularVehicleAsset* NewModularVehicle = StaticFactoryCreateNew(Class, InParent, Name, Flags, Context, Warn);
+
+	ensure(false); // @todo(chaos) implement proper copy
+	//for (ModularVehicleTuple & ModularVehicleData : ModularVehicleList)
+	//{
+	//	NewModularVehicle->AppendGeometry(*ModularVehicleData.Get<0>(), false, ModularVehicleData.Get<2>());
+	//}
+
+	ensure(false); // @todo(chaos) implement proper simulation prep
+	//ModularVehicleAlgo::PrepareForSimulation(NewModularVehicle->GetModularVehicle().Get());
+
+	NewModularVehicle->Modify();
+	return NewModularVehicle;
+}
+
+#undef LOCTEXT_NAMESPACE
+
+
+
diff --git a/Engine/Plugins/Experimental/ChaosModularVehicle/Source/ChaosModularVehicleEditor/Private/ChaosModularVehicle/ModularVehicleAssetThumbnailRenderer.cpp b/Engine/Plugins/Experimental/ChaosModularVehicle/Source/ChaosModularVehicleEditor/Private/ChaosModularVehicle/ModularVehicleAssetThumbnailRenderer.cpp
new file mode 100644
--- /dev/null
+++ b/Engine/Plugins/Experimental/ChaosModularVehicle/Source/ChaosModularVehicleEditor/Private/ChaosModularVehicle/ModularVehicleAssetThumbnailRenderer.cpp
@@ -0,0 +1,51 @@
+// Copyright Epic Games, Inc. All Rights Reserved.
+
+#include "ChaosModularVehicle/ModularVehicleAssetThumbnailRenderer.h"
+#include "Misc/App.h"
+#include "ShowFlags.h"
+#include "SceneView.h"
+#include "ChaosModularVehicle/ModularVehicleAssetThumbnailScene.h"
+#include "ChaosModularVehicle/ModularVehicleAsset.h"
+
+UModularVehicleAssetThumbnailRenderer::UModularVehicleAssetThumbnailRenderer(const FObjectInitializer& ObjectInitializer)
+: Super(ObjectInitializer)
+{
+	ThumbnailScene = nullptr;
+}
+
+void UModularVehicleAssetThumbnailRenderer::Draw(UObject* Object, int32 X, int32 Y, uint32 Width, uint32 Height, FRenderTarget* RenderTarget, FCanvas* Canvas, bool bAdditionalViewFamily)
+{
+	UModularVehicleAsset* ModularVehicleAsset = Cast<UModularVehicleAsset>(Object);
+	if (IsValid(ModularVehicleAsset))
+	{
+		//if (ThumbnailScene == nullptr)
+		//{
+		//	ThumbnailScene = new FModularVehicleAssetThumbnailScene();
+		//}
+
+		//ThumbnailScene->SetModularVehicleAsset(ModularVehicleAsset);
+		//ThumbnailScene->GetScene()->UpdateSpeedTreeWind(0.0);
+
+		////FSceneViewFamilyContext ViewFamily(FSceneViewFamily::ConstructionValues(RenderTarget, ThumbnailScene->GetScene(), FEngineShowFlags(ESFIM_Game))
+		////	.SetTime(UThumbnailRenderer::GetTime())
+		////	.SetAdditionalViewFamily(bAdditionalViewFamily));
+
+		//ViewFamily.EngineShowFlags.DisableAdvancedFeatures();
+		//ViewFamily.EngineShowFlags.MotionBlur = 0;
+		//ViewFamily.EngineShowFlags.LOD = 0;
+
+		//RenderViewFamily(Canvas, &ViewFamily, ThumbnailScene->CreateView(&ViewFamily, X, Y, Width, Height));
+		//ThumbnailScene->SetModularVehicleAsset(nullptr);
+	}
+}
+
+void UModularVehicleAssetThumbnailRenderer::BeginDestroy()
+{
+	if (ThumbnailScene != nullptr)
+	{
+		delete ThumbnailScene;
+		ThumbnailScene = nullptr;
+	}
+
+	Super::BeginDestroy();
+}
diff --git a/Engine/Plugins/Experimental/ChaosModularVehicle/Source/ChaosModularVehicleEditor/Private/ChaosModularVehicle/ModularVehicleAssetThumbnailScene.cpp b/Engine/Plugins/Experimental/ChaosModularVehicle/Source/ChaosModularVehicleEditor/Private/ChaosModularVehicle/ModularVehicleAssetThumbnailScene.cpp
new file mode 100644
--- /dev/null
+++ b/Engine/Plugins/Experimental/ChaosModularVehicle/Source/ChaosModularVehicleEditor/Private/ChaosModularVehicle/ModularVehicleAssetThumbnailScene.cpp
@@ -0,0 +1,73 @@
+// Copyright Epic Games, Inc. All Rights Reserved.
+
+#include "ChaosModularVehicle/ModularVehicleAssetThumbnailScene.h"
+#include "ThumbnailRendering/SceneThumbnailInfo.h"
+#include "ChaosModularVehicle/ModularVehicleAsset.h"
+#include "ChaosModularVehicle/ModularVehicleComponent.h"
+#include "ChaosModularVehicle/ModularVehiclePawn.h"
+
+FModularVehicleAssetThumbnailScene::FModularVehicleAssetThumbnailScene()
+{
+	bForceAllUsedMipsResident = false;
+
+	// Create preview actor
+	// checked
+	FActorSpawnParameters SpawnInfo;
+	SpawnInfo.SpawnCollisionHandlingOverride = ESpawnActorCollisionHandlingMethod::AlwaysSpawn;
+	SpawnInfo.bNoFail = true;
+	SpawnInfo.ObjectFlags = RF_Transient;
+
+	PreviewActor = GetWorld()->SpawnActor<AModularVehiclePawn>(SpawnInfo);
+	PreviewActor->GetModularVehicleComponent()->SetMobility(EComponentMobility::Movable);
+	PreviewActor->SetActorEnableCollision(false);
+}
+
+void FModularVehicleAssetThumbnailScene::SetModularVehicleAsset(UModularVehicleAsset* ModularVehicleAsset)
+{
+	//PreviewActor->GetModularVehicleComponent()->SetVehicleAsset(ModularVehicleAsset);
+
+	if (ModularVehicleAsset)
+	{
+		FTransform MeshTransform = FTransform::Identity;
+
+		PreviewActor->SetActorLocation(FVector(0, 0, 0), false);
+		PreviewActor->GetModularVehicleComponent()->UpdateBounds();
+
+		// Center the mesh at the world origin then offset to put it on top of the plane
+		const float BoundsZOffset = GetBoundsZOffset(PreviewActor->GetModularVehicleComponent()->Bounds);
+		PreviewActor->SetActorLocation(-PreviewActor->GetModularVehicleComponent()->Bounds.Origin + FVector(0, 0, BoundsZOffset), false);
+		PreviewActor->GetModularVehicleComponent()->RecreateRenderState_Concurrent();
+	}
+}
+
+void FModularVehicleAssetThumbnailScene::GetViewMatrixParameters(const float InFOVDegrees, FVector& OutOrigin, float& OutOrbitPitch, float& OutOrbitYaw, float& OutOrbitZoom) const
+{
+	check(PreviewActor);
+	check(PreviewActor->GetModularVehicleComponent());
+	//check(PreviewActor->GetModularVehicleComponent()->GetVehicleAsset());
+
+	const float HalfFOVRadians = FMath::DegreesToRadians<float>(InFOVDegrees) * 0.5f;
+	// Add extra size to view slightly outside of the sphere to compensate for perspective
+	const float HalfMeshSize = PreviewActor->GetModularVehicleComponent()->Bounds.SphereRadius * 1.15;
+	const float BoundsZOffset = GetBoundsZOffset(PreviewActor->GetModularVehicleComponent()->Bounds);
+	const float TargetDistance = HalfMeshSize / FMath::Tan(HalfFOVRadians);
+
+	USceneThumbnailInfo* ThumbnailInfo;// = Cast<USceneThumbnailInfo>(PreviewActor->GetModularVehicleComponent()->GetVehicleAsset()->ThumbnailInfo);
+	//if (ThumbnailInfo)
+	//{
+	//	if (TargetDistance + ThumbnailInfo->OrbitZoom < 0)
+	//	{
+	//		ThumbnailInfo->OrbitZoom = -TargetDistance;
+	//	}
+	//}
+	//else
+	{
+		ThumbnailInfo = USceneThumbnailInfo::StaticClass()->GetDefaultObject<USceneThumbnailInfo>();
+	}
+
+	OutOrigin = FVector(0, 0, -BoundsZOffset);
+	OutOrbitPitch = ThumbnailInfo->OrbitPitch;
+	OutOrbitYaw = ThumbnailInfo->OrbitYaw;
+	OutOrbitZoom = TargetDistance + ThumbnailInfo->OrbitZoom;
+}
+
diff --git a/Engine/Plugins/Experimental/ChaosModularVehicle/Source/ChaosModularVehicleEditor/Private/ChaosModularVehicle/ModularVehicleBPInterface.cpp b/Engine/Plugins/Experimental/ChaosModularVehicle/Source/ChaosModularVehicleEditor/Private/ChaosModularVehicle/ModularVehicleBPInterface.cpp
new file mode 100644
--- /dev/null
+++ b/Engine/Plugins/Experimental/ChaosModularVehicle/Source/ChaosModularVehicleEditor/Private/ChaosModularVehicle/ModularVehicleBPInterface.cpp
@@ -0,0 +1,21 @@
+// Copyright Epic Games, Inc. All Rights Reserved.
+
+#include "ChaosModularVehicle/ModularVehicleBPInterface.h"
+#include "ChaosModularVehicle/ModularToolkitBlueprintLibrary.h"
+#include "ChaosModularVehicle/ModularVehiclePawn.h"
+
+
+UModularVehicleBPInterface::UModularVehicleBPInterface(const FObjectInitializer& ObjectInitializer)
+	: Super(ObjectInitializer)
+{
+}
+
+void UModularVehicleBPInterface::AppendSimModule(AModularVehiclePawn* ExistingActorInOut, int32 ParentIndex, const AActor* Actor, int32& OutIndex)
+{
+	OutIndex = UModularToolkitBlueprintLibrary::AddChild(ExistingActorInOut, ParentIndex, Actor);
+}
+
+void UModularVehicleBPInterface::ConvertToVehicle(const FString& NewAssetPath, const TArray<AActor*>& Actors)
+{
+	UModularToolkitBlueprintLibrary::ConvertActorsToModularVehicle(NewAssetPath, Actors);
+}
diff --git a/Engine/Plugins/Experimental/ChaosModularVehicle/Source/ChaosModularVehicleEditor/Public/ChaosModularVehicle/AssetTypeActions_ModularVehicleAsset.h b/Engine/Plugins/Experimental/ChaosModularVehicle/Source/ChaosModularVehicleEditor/Public/ChaosModularVehicle/AssetTypeActions_ModularVehicleAsset.h
new file mode 100644
--- /dev/null
+++ b/Engine/Plugins/Experimental/ChaosModularVehicle/Source/ChaosModularVehicleEditor/Public/ChaosModularVehicle/AssetTypeActions_ModularVehicleAsset.h
@@ -0,0 +1,19 @@
+// Copyright Epic Games, Inc. All Rights Reserved.
+
+#pragma once
+
+#include "CoreMinimal.h"
+#include "Toolkits/IToolkitHost.h"
+#include "AssetTypeActions_Base.h"
+
+class FAssetTypeActions_ModularVehicleAsset : public FAssetTypeActions_Base
+{
+public:
+	// IAssetTypeActions Implementation
+	virtual FText GetName() const override { return NSLOCTEXT("AssetTypeActions", "AssetTypeActions_ModularVehicleAsset", "ModularVehicle Asset"); }
+	virtual void GetActions(const TArray<UObject*>& InObjects, struct FToolMenuSection& Section) override;
+	virtual FColor GetTypeColor() const override { return FColor(255, 192, 128); }
+	virtual UClass* GetSupportedClass() const override;
+	virtual uint32 GetCategories() override { return EAssetTypeCategories::Physics; }
+	virtual UThumbnailInfo* GetThumbnailInfo(UObject* Asset) const override;
+};
diff --git a/Engine/Plugins/Experimental/ChaosModularVehicle/Source/ChaosModularVehicleEditor/Public/ChaosModularVehicle/ChaosModularVehicleCommands.h b/Engine/Plugins/Experimental/ChaosModularVehicle/Source/ChaosModularVehicleEditor/Public/ChaosModularVehicle/ChaosModularVehicleCommands.h
new file mode 100644
--- /dev/null
+++ b/Engine/Plugins/Experimental/ChaosModularVehicle/Source/ChaosModularVehicleEditor/Public/ChaosModularVehicle/ChaosModularVehicleCommands.h
@@ -0,0 +1,19 @@
+// Copyright Epic Games, Inc. All Rights Reserved.
+
+#pragma once
+#include "Engine/World.h"
+
+/**
+* The public interface to this module
+*/
+class CHAOSMODULARVEHICLEEDITOR_API FChaosModularVehicleCommands
+{
+public:
+
+
+	/**
+	*  Import file. 
+	*/
+	static void ImportFile(const FString & FileName);
+
+};
\ No newline at end of file
diff --git a/Engine/Plugins/Experimental/ChaosModularVehicle/Source/ChaosModularVehicleEditor/Public/ChaosModularVehicle/ChaosModularVehicleEditorPlugin.h b/Engine/Plugins/Experimental/ChaosModularVehicle/Source/ChaosModularVehicleEditor/Public/ChaosModularVehicle/ChaosModularVehicleEditorPlugin.h
new file mode 100644
--- /dev/null
+++ b/Engine/Plugins/Experimental/ChaosModularVehicle/Source/ChaosModularVehicleEditor/Public/ChaosModularVehicle/ChaosModularVehicleEditorPlugin.h
@@ -0,0 +1,58 @@
+// Copyright Epic Games, Inc. All Rights Reserved.
+
+#pragma once
+
+#include "CoreMinimal.h"
+#include "Modules/ModuleInterface.h"
+#include "Modules/ModuleManager.h"
+
+class ISlateStyle;
+class FSlateStyleSet;
+class FAssetTypeActions_ModularVehicleAsset;
+
+
+/**
+ * The public interface to this module
+ */
+class IChaosModularVehicleEditorPlugin : public IModuleInterface
+{
+public:
+
+	/**
+	 * Singleton-like access to this module's interface.  This is just for convenience!
+	 * Beware of calling this during the shutdown phase, though.  Your module might have been unloaded already.
+	 *
+	 * @return Returns singleton instance, loading the module on demand if needed
+	 */
+	static inline IChaosModularVehicleEditorPlugin& Get()
+	{
+		return FModuleManager::LoadModuleChecked< IChaosModularVehicleEditorPlugin >( "ChaosModularVehicleEditor" );
+	}
+
+	/**
+	 * Checks to see if this module is loaded and ready.  It is only valid to call Get() if IsAvailable() returns true.
+	 *
+	 * @return True if the module is loaded and ready to use
+	 */
+	static inline bool IsAvailable()
+	{
+		return FModuleManager::Get().IsModuleLoaded( "ChaosModularVehicleEditor" );
+	}
+
+	TSharedPtr<FSlateStyleSet> GetStyleSet() { return StyleSet; }
+
+	static FName GetEditorStyleName();
+	static const ISlateStyle* GetEditorStyle();
+
+protected:
+	void RegisterMenus();
+
+	TArray<IConsoleObject*> EditorCommands;
+
+	// Asset actions for new asset types
+	FAssetTypeActions_ModularVehicleAsset* ModularVehicleAssetActions;
+
+	// Styleset for flesh tool brushes/fonts etc.
+	TSharedPtr<FSlateStyleSet> StyleSet;
+};
+
diff --git a/Engine/Plugins/Experimental/ChaosModularVehicle/Source/ChaosModularVehicleEditor/Public/ChaosModularVehicle/ModularToolkitBlueprintLibrary.h b/Engine/Plugins/Experimental/ChaosModularVehicle/Source/ChaosModularVehicleEditor/Public/ChaosModularVehicle/ModularToolkitBlueprintLibrary.h
new file mode 100644
--- /dev/null
+++ b/Engine/Plugins/Experimental/ChaosModularVehicle/Source/ChaosModularVehicleEditor/Public/ChaosModularVehicle/ModularToolkitBlueprintLibrary.h
@@ -0,0 +1,78 @@
+// Copyright Epic Games, Inc. All Rights Reserved.
+
+#pragma once
+
+#include "CoreMinimal.h"
+#include "UObject/ObjectMacros.h"
+#include "Kismet/BlueprintFunctionLibrary.h"
+#include "EditorSubsystem.h"
+#include "Engine/EngineTypes.h"
+
+#include "ModularToolkitBlueprintLibrary.generated.h"
+
+DECLARE_LOG_CATEGORY_EXTERN(LogModularVehicleToolkit, Log, All);
+
+class AModularVehiclePawn;
+class UModularVehicle;
+class UGeometryCollection;
+class UModularVehicleComponent;
+
+UCLASS(ClassGroup = (Physics))
+class CHAOSMODULARVEHICLEEDITOR_API UModularToolkitBlueprintLibrary : public UObject
+{
+	
+	GENERATED_UCLASS_BODY()
+
+	enum class EModuleOperation : uint8
+	{
+		Add,
+		Insert,
+		Remove
+	};
+
+	/**
+		* BP Function, Searches for static mesh and geometry collection components in source actors and adds them into an existing GeometryCollectionActor
+		*/
+//	UFUNCTION(BlueprintCallable, Category = "ModularVehicle")
+	static void AddActorToExistingModularVehicle(AModularVehiclePawn* ExistingActorInOut, const AActor* Actor);
+
+	static void AddActorsToExistingModularVehicle(AModularVehiclePawn* ExistingActorInOut, const TArray<AActor*>& Actors);
+
+
+
+	/**
+		* Searches for static mesh and geometry collection components in source actors and combines them into one new GeometryCollectionActor
+		*/
+	static AModularVehiclePawn* ConvertActorsToModularVehicle(const FString& InAssetPath, const TArray<AActor*>& Actors);
+
+	/**
+		* Creates a new geometry collection actor from scratch
+		*/
+	static class AModularVehiclePawn* CreateNewModularVehicleActor(const FString& InAssetPath, const FTransform& Transform);
+	
+	static class AModularVehiclePawn* SetupRestCollectionOnActor(const FString& InAssetPath, AModularVehiclePawn* ModularVehicleActor, const FTransform& Transform);
+
+	static bool Attach(UModularVehicleComponent* ModularVehicleComponent, UClass* AttachClass, UClass* ToClass);
+
+	static bool AttachNearest(UModularVehicleComponent* ModularVehicleComponent, UClass* AttachClass, UClass* ToClass);
+
+	// Add new vehicle piece as child of specified parent piece index, returns new piece's index
+	static int AddChild(AModularVehiclePawn* ExistingVehicleActor, int32 ParentIndex, const AActor* Actor);
+
+	// inserts vehicle piece between existing parent and child at the specified child index, returns new piece's index
+	static int InsertAbove(AModularVehiclePawn* ExistingVehicleActor, int32 ChildIndex, const AActor* Actor);
+
+	// remove existing vehicle piece, moving the deleted pieces children up to the deleted pieces parent
+	static int RemoveAt(AModularVehiclePawn* ExistingVehicleActor, int Index);
+
+
+private:
+	static void AddSingleRootNodeIfRequired(UGeometryCollection* GeometryCollectionObject);
+	static void AddAdditionalAttributesIfRequired(UGeometryCollection* GeometryCollectionObject);
+	static void FlattenHierarchy(UGeometryCollection* GeometryCollectionObject);
+
+	static ULevel* GetSelectedLevel();
+	static AActor* AddActor(ULevel* InLevel, UClass* Class);
+	static int UpdateSimulationTree(AModularVehiclePawn* ExistingVehicle, EModuleOperation Operation, int Index, const AActor* Actor);
+
+};
diff --git a/Engine/Plugins/Experimental/ChaosModularVehicle/Source/ChaosModularVehicleEditor/Public/ChaosModularVehicle/ModularVehicleAssetConversion.h b/Engine/Plugins/Experimental/ChaosModularVehicle/Source/ChaosModularVehicleEditor/Public/ChaosModularVehicle/ModularVehicleAssetConversion.h
new file mode 100644
--- /dev/null
+++ b/Engine/Plugins/Experimental/ChaosModularVehicle/Source/ChaosModularVehicleEditor/Public/ChaosModularVehicle/ModularVehicleAssetConversion.h
@@ -0,0 +1,18 @@
+// Copyright Epic Games, Inc. All Rights Reserved.
+
+#pragma once
+#include "Engine/World.h"
+
+
+/**
+* The public interface to this module
+*/
+class CHAOSMODULARVEHICLEEDITOR_API FModularVehicleAssetConversion
+{
+public:
+
+	/**
+	*/
+	static void ImportFile(FString& FileName);
+
+};
\ No newline at end of file
diff --git a/Engine/Plugins/Experimental/ChaosModularVehicle/Source/ChaosModularVehicleEditor/Public/ChaosModularVehicle/ModularVehicleAssetFactory.h b/Engine/Plugins/Experimental/ChaosModularVehicle/Source/ChaosModularVehicleEditor/Public/ChaosModularVehicle/ModularVehicleAssetFactory.h
new file mode 100644
--- /dev/null
+++ b/Engine/Plugins/Experimental/ChaosModularVehicle/Source/ChaosModularVehicleEditor/Public/ChaosModularVehicle/ModularVehicleAssetFactory.h
@@ -0,0 +1,33 @@
+// Copyright Epic Games, Inc. All Rights Reserved.
+#pragma once
+
+#include "CoreMinimal.h"
+#include "UObject/ObjectMacros.h"
+#include "Factories/Factory.h"
+
+#include "ModularVehicleAssetFactory.generated.h"
+
+class UModularVehicleAsset;
+class UModularVehicleComponent;
+
+typedef TTuple<const UModularVehicleAsset *, const UModularVehicleComponent *, FTransform> ModularVehicleTuple;
+
+
+/**
+* Factory for ModularVehicleAsset
+*/
+
+UCLASS()
+class CHAOSMODULARVEHICLEEDITOR_API UModularVehicleAssetFactory : public UFactory
+{
+    GENERATED_UCLASS_BODY()
+
+	//~ Begin UFactory Interface
+	virtual UObject* FactoryCreateNew(UClass* Class, UObject* InParent, FName Name, EObjectFlags Flags, UObject* Context, FFeedbackContext* Warn) override;
+	//~ Begin UFactory Interface	
+
+	static UModularVehicleAsset* StaticFactoryCreateNew(UClass* Class, UObject* InParent, FName Name, EObjectFlags Flags, UObject* Context, FFeedbackContext* Warn);
+
+};
+
+
diff --git a/Engine/Plugins/Experimental/ChaosModularVehicle/Source/ChaosModularVehicleEditor/Public/ChaosModularVehicle/ModularVehicleAssetThumbnailRenderer.h b/Engine/Plugins/Experimental/ChaosModularVehicle/Source/ChaosModularVehicleEditor/Public/ChaosModularVehicle/ModularVehicleAssetThumbnailRenderer.h
new file mode 100644
--- /dev/null
+++ b/Engine/Plugins/Experimental/ChaosModularVehicle/Source/ChaosModularVehicleEditor/Public/ChaosModularVehicle/ModularVehicleAssetThumbnailRenderer.h
@@ -0,0 +1,33 @@
+// Copyright Epic Games, Inc. All Rights Reserved.
+
+/**
+*
+* This thumbnail renderer displays a given GeometryCollection
+*/
+
+#pragma once
+
+#include "CoreMinimal.h"
+#include "UObject/ObjectMacros.h"
+#include "ThumbnailRendering/DefaultSizedThumbnailRenderer.h"
+#include "ModularVehicleAssetThumbnailRenderer.generated.h"
+
+class FCanvas;
+class FModularVehicleAssetThumbnailScene;
+class FRenderTarget;
+
+UCLASS(config = Editor, MinimalAPI)
+class UModularVehicleAssetThumbnailRenderer : public UDefaultSizedThumbnailRenderer
+{
+	GENERATED_UCLASS_BODY()
+
+	// Begin UThumbnailRenderer Object
+	virtual void Draw(UObject* Object, int32 X, int32 Y, uint32 Width, uint32 Height, FRenderTarget* RenderTarget, FCanvas* Canvas, bool bAdditionalViewFamily) override;
+	// End UThumbnailRenderer Object
+
+	// Begin UObject implementation
+	virtual void BeginDestroy() override;
+	// End UObject implementation
+private:
+	FModularVehicleAssetThumbnailScene* ThumbnailScene;
+};
diff --git a/Engine/Plugins/Experimental/ChaosModularVehicle/Source/ChaosModularVehicleEditor/Public/ChaosModularVehicle/ModularVehicleAssetThumbnailScene.h b/Engine/Plugins/Experimental/ChaosModularVehicle/Source/ChaosModularVehicleEditor/Public/ChaosModularVehicle/ModularVehicleAssetThumbnailScene.h
new file mode 100644
--- /dev/null
+++ b/Engine/Plugins/Experimental/ChaosModularVehicle/Source/ChaosModularVehicleEditor/Public/ChaosModularVehicle/ModularVehicleAssetThumbnailScene.h
@@ -0,0 +1,27 @@
+// Copyright Epic Games, Inc. All Rights Reserved.
+
+#pragma once
+
+#include "CoreMinimal.h"
+#include "ThumbnailHelpers.h"
+
+class AModularVehiclePawn;
+class UModularVehicleAsset;
+
+class FModularVehicleAssetThumbnailScene : public FThumbnailPreviewScene
+{
+public:
+	/** Constructor */
+	FModularVehicleAssetThumbnailScene();
+
+	/** Sets the geometry collection to use in the next CreateView() */
+	void SetModularVehicleAsset(UModularVehicleAsset* ModularVehicleAsset);
+
+protected:
+	// FThumbnailPreviewScene implementation
+	virtual void GetViewMatrixParameters(const float InFOVDegrees, FVector& OutOrigin, float& OutOrbitPitch, float& OutOrbitYaw, float& OutOrbitZoom) const override;
+
+private:
+	/** The actor used to display all geometry collection thumbnails */
+	AModularVehiclePawn* PreviewActor;
+};
diff --git a/Engine/Plugins/Experimental/ChaosModularVehicle/Source/ChaosModularVehicleEditor/Public/ChaosModularVehicle/ModularVehicleBPInterface.h b/Engine/Plugins/Experimental/ChaosModularVehicle/Source/ChaosModularVehicleEditor/Public/ChaosModularVehicle/ModularVehicleBPInterface.h
new file mode 100644
--- /dev/null
+++ b/Engine/Plugins/Experimental/ChaosModularVehicle/Source/ChaosModularVehicleEditor/Public/ChaosModularVehicle/ModularVehicleBPInterface.h
@@ -0,0 +1,30 @@
+// Copyright Epic Games, Inc. All Rights Reserved.
+#pragma once
+
+#include "CoreMinimal.h"
+#include "GameFramework/Actor.h"
+#include "Kismet/BlueprintFunctionLibrary.h"
+#include "ModularVehicleBPInterface.generated.h"
+
+class AModularVehiclePawn;
+class AActor;
+
+// Test to see exposing of BP function - not working in editor portion of module!
+UCLASS(ClassGroup = (Physics))
+class CHAOSMODULARVEHICLEEDITOR_API UModularVehicleBPInterface : public UBlueprintFunctionLibrary
+{
+	GENERATED_UCLASS_BODY()
+
+	UFUNCTION(BlueprintCallable, Category = "ModularVehicle")
+	static void AppendSimModule(AModularVehiclePawn* ExistingActorInOut, int32 ParentIndex, const AActor* Actor, int32& OutIndex);
+
+	UFUNCTION(BlueprintCallable, Category = "ModularVehicle")
+	static void ConvertToVehicle(const FString& NewAssetPath, const TArray<AActor*>& Actors);
+
+	UFUNCTION(BlueprintCallable, Category = "ModularVehicle")
+	static float GenerateAsset()
+	{
+		return false;
+	}
+
+};
diff --git a/Engine/Plugins/Experimental/ChaosModularVehicle/Source/ChaosModularVehicleEngine/Private/ChaosModularVehicle/ChaosSimModuleManager.cpp b/Engine/Plugins/Experimental/ChaosModularVehicle/Source/ChaosModularVehicleEngine/Private/ChaosModularVehicle/ChaosSimModuleManager.cpp
new file mode 100644
--- /dev/null
+++ b/Engine/Plugins/Experimental/ChaosModularVehicle/Source/ChaosModularVehicleEngine/Private/ChaosModularVehicle/ChaosSimModuleManager.cpp
@@ -0,0 +1,368 @@
+//// Copyright Epic Games, Inc. All Rights Reserved.
+//
+#include "ChaosModularVehicle/ChaosSimModuleManager.h"
+#include "ChaosModularVehicle/ModularVehicleComponent.h"
+#include "ChaosModularVehicle/ModularVehicleBaseComponent.h"
+#include "PBDRigidsSolver.h"
+#include "GameFramework/HUD.h" // for ShowDebugInfo
+#include "Physics/Experimental/PhysScene_Chaos.h"
+
+
+TMap<FPhysScene*, FChaosSimModuleManager*> FChaosSimModuleManager::SceneToModuleManagerMap;
+
+FDelegateHandle FChaosSimModuleManager::OnPostWorldInitializationHandle;
+FDelegateHandle FChaosSimModuleManager::OnWorldCleanupHandle;
+
+extern FSimModuleDebugParams GSimModuleDebugParams;
+
+bool FChaosSimModuleManager::GInitialized = false;
+
+
+FChaosSimModuleManager::FChaosSimModuleManager(FPhysScene* PhysScene)
+	: Scene(*PhysScene)
+	, AsyncCallback(nullptr)
+	, Timestamp(0)
+
+{
+	check(PhysScene);
+	
+	if (!GInitialized)
+	{
+		GInitialized = true;
+		// PhysScene->GetOwningWorld() is always null here, the world is being setup too late to be of use
+		// therefore setup these global world delegates that will callback when everything is setup so registering
+		// the physics solver Async Callback will succeed
+		OnPostWorldInitializationHandle = FWorldDelegates::OnPostWorldInitialization.AddStatic(&FChaosSimModuleManager::OnPostWorldInitialization);
+		OnWorldCleanupHandle = FWorldDelegates::OnWorldCleanup.AddStatic(&FChaosSimModuleManager::OnWorldCleanup);
+
+		if (!IsRunningDedicatedServer())
+		{
+			AHUD::OnShowDebugInfo.AddStatic(&FChaosSimModuleManager::OnShowDebugInfo);
+		}
+	}
+
+	ensure(FChaosSimModuleManager::SceneToModuleManagerMap.Find(PhysScene) == nullptr);	// double registration with same scene, will cause a leak
+
+	// Add to Scene-To-Manager map
+	FChaosSimModuleManager::SceneToModuleManagerMap.Add(PhysScene, this);
+}
+
+FChaosSimModuleManager::~FChaosSimModuleManager()
+{
+	while (GCVehicles.Num() > 0)
+	{
+		RemoveVehicle(GCVehicles.Last());
+	}
+
+	while (CUVehicles.Num() > 0)
+	{
+		RemoveVehicle(CUVehicles.Last());
+	}
+}
+
+void FChaosSimModuleManager::OnPostWorldInitialization(UWorld* InWorld, const UWorld::InitializationValues)
+{
+	FChaosSimModuleManager* Manager = FChaosSimModuleManager::GetManagerFromScene(InWorld->GetPhysicsScene());
+	if (Manager)
+	{
+		Manager->RegisterCallbacks(InWorld);
+	}
+}
+
+void FChaosSimModuleManager::OnWorldCleanup(UWorld* InWorld, bool bSessionEnded, bool bCleanupResources)
+{
+	FChaosSimModuleManager* Manager = FChaosSimModuleManager::GetManagerFromScene(InWorld->GetPhysicsScene());
+	if (Manager)
+	{
+		Manager->UnregisterCallbacks();
+	}
+}
+
+void FChaosSimModuleManager::OnShowDebugInfo(AHUD* HUD, UCanvas* Canvas, const FDebugDisplayInfo& DisplayInfo, float& YL, float& YPos)
+{
+	static const FName NAME_VehicleModuleManager("VehicleModuleManager");
+	if (Canvas && HUD->ShouldDisplayDebug(NAME_VehicleModuleManager))
+	{
+		FChaosSimModuleManager* Manager = FChaosSimModuleManager::GetManagerFromScene(HUD->GetWorld()->GetPhysicsScene());
+
+		int ShowVehicleIndex = 0;
+		if (!Manager->CUVehicles.IsEmpty())
+		{
+			if (Manager->CUVehicles[ShowVehicleIndex].IsValid())
+			{
+				Manager->CUVehicles[ShowVehicleIndex]->ShowDebugInfo(HUD, Canvas, DisplayInfo, YL, YPos);
+			}
+		}
+
+		if (!Manager->GCVehicles.IsEmpty())
+		{
+			if (Manager->GCVehicles[ShowVehicleIndex].IsValid())
+			{
+				Manager->GCVehicles[ShowVehicleIndex]->ShowDebugInfo(HUD, Canvas, DisplayInfo, YL, YPos);
+			}
+		}
+
+
+	}
+}
+
+void FChaosSimModuleManager::DetachFromPhysScene(FPhysScene* PhysScene)
+{
+	if (AsyncCallback)
+	{
+		UnregisterCallbacks();
+	}
+
+	if (PhysScene->GetOwningWorld())
+	{
+		PhysScene->GetOwningWorld()->OnWorldBeginPlay.RemoveAll(this);
+	}
+
+	FChaosSimModuleManager::SceneToModuleManagerMap.Remove(PhysScene);
+}
+
+
+void FChaosSimModuleManager::RegisterCallbacks(UWorld* InWorld)
+{
+	OnPhysScenePreTickHandle = Scene.OnPhysScenePreTick.AddRaw(this, &FChaosSimModuleManager::Update);
+	OnPhysScenePostTickHandle = Scene.OnPhysScenePostTick.AddRaw(this, &FChaosSimModuleManager::PostUpdate);
+
+	// Set up our async object manager to handle async ticking and marshaling
+	check(AsyncCallback == nullptr);
+	AsyncCallback = Scene.GetSolver()->CreateAndRegisterSimCallbackObject_External<FChaosSimModuleManagerAsyncCallback>();
+}
+
+void FChaosSimModuleManager::UnregisterCallbacks()
+{
+	Scene.OnPhysScenePreTick.Remove(OnPhysScenePreTickHandle);
+	Scene.OnPhysScenePostTick.Remove(OnPhysScenePostTickHandle);
+	
+	if (AsyncCallback)
+	{
+		Scene.GetSolver()->UnregisterAndFreeSimCallbackObject_External(AsyncCallback);
+		AsyncCallback = nullptr;
+	}
+
+}
+
+FChaosSimModuleManager* FChaosSimModuleManager::GetManagerFromScene(FPhysScene* PhysScene)
+{
+	FChaosSimModuleManager* Manager = nullptr;
+	FChaosSimModuleManager** ManagerPtr = SceneToModuleManagerMap.Find(PhysScene);
+	if (ManagerPtr != nullptr)
+	{
+		Manager = *ManagerPtr;
+	}
+	return Manager;
+}
+
+void FChaosSimModuleManager::AddVehicle(TWeakObjectPtr<UModularVehicleComponent> Vehicle)
+{
+	check(Vehicle != NULL);
+	check(Vehicle->PhysicsVehicleOutput());
+	check(AsyncCallback);
+
+	GCVehicles.Add(Vehicle);
+}
+
+void FChaosSimModuleManager::RemoveVehicle(TWeakObjectPtr<UModularVehicleComponent> Vehicle)
+{
+	check(Vehicle != NULL);
+
+	GCVehicles.Remove(Vehicle);
+}
+
+void FChaosSimModuleManager::AddVehicle(TWeakObjectPtr < UModularVehicleBaseComponent> Vehicle)
+{
+	check(Vehicle != NULL);
+	check(Vehicle->PhysicsVehicleOutput());
+	check(AsyncCallback);
+	CUVehicles.Add(Vehicle);
+}
+
+void FChaosSimModuleManager::RemoveVehicle(TWeakObjectPtr<UModularVehicleBaseComponent> Vehicle)
+{
+	check(Vehicle != NULL);
+
+	CUVehicles.Remove(Vehicle);
+}
+
+void FChaosSimModuleManager::ScenePreTick(FPhysScene* PhysScene, float DeltaTime)
+{
+	for (int32 i = 0; i < GCVehicles.Num(); ++i)
+	{
+		GCVehicles[i]->PreTickGT(DeltaTime);
+	}
+
+	for (int32 i = 0; i < CUVehicles.Num(); ++i)
+	{
+		CUVehicles[i]->PreTickGT(DeltaTime);
+	}
+}
+
+void FChaosSimModuleManager::Update(FPhysScene* PhysScene, float DeltaTime)
+{
+	UWorld* World = Scene.GetOwningWorld();
+
+	SubStepCount = 0;
+
+	ScenePreTick(PhysScene, DeltaTime);
+
+	ParallelUpdateVehicles(DeltaTime);
+
+	if (World)
+	{
+		FChaosSimModuleManagerAsyncInput* AsyncInput = AsyncCallback->GetProducerInputData_External();
+
+		for (TWeakObjectPtr<UModularVehicleComponent> Vehicle : GCVehicles)
+		{
+			Vehicle->Update(DeltaTime);
+			Vehicle->FinalizeSimCallbackData(*AsyncInput);
+		}
+
+
+		for (TWeakObjectPtr<UModularVehicleBaseComponent> Vehicle : CUVehicles)
+		{
+			Vehicle->Update(DeltaTime);
+			Vehicle->FinalizeSimCallbackData(*AsyncInput);
+		}
+	}
+}
+
+
+void FChaosSimModuleManager::PostUpdate(FChaosScene* PhysScene)
+{
+}
+
+
+void FChaosSimModuleManager::ParallelUpdateVehicles(float DeltaSeconds)
+{
+
+	FChaosSimModuleManagerAsyncInput* AsyncInput = AsyncCallback->GetProducerInputData_External();
+
+	AsyncInput->Reset();	// only want latest frame's data
+
+	{
+		// We pass pointers from TArray so this reserve is critical. Otherwise realloc happens
+		AsyncInput->VehicleInputs.Reserve(CUVehicles.Num() + GCVehicles.Num());
+		AsyncInput->Timestamp = Timestamp;
+		AsyncInput->World = Scene.GetOwningWorld();
+	}
+
+	// Grab all outputs for processing, even future ones for interpolation.
+	{
+		Chaos::TSimCallbackOutputHandle<FChaosSimModuleManagerAsyncOutput> AsyncOutputLatest;
+		while ((AsyncOutputLatest = AsyncCallback->PopFutureOutputData_External()))
+		{
+			PendingOutputs.Emplace(MoveTemp(AsyncOutputLatest));
+		}
+	}
+
+	// Since we are in pre-physics, delta seconds is not accounted for in external time yet
+	const float ResultsTime = AsyncCallback->GetSolver()->GetPhysicsResultsTime_External() + DeltaSeconds;
+
+	// Find index of first non-consumable output (first one after current time)
+	int32 LastOutputIdx = 0;
+	for (; LastOutputIdx < PendingOutputs.Num(); ++LastOutputIdx)
+	{
+		if (PendingOutputs[LastOutputIdx]->InternalTime > ResultsTime)
+		{
+			break;
+		}
+	}
+
+	// Cache the last consumed output for interpolation
+	if (LastOutputIdx > 0)
+	{
+		LatestOutput = MoveTemp(PendingOutputs[LastOutputIdx - 1]);
+	}
+
+	// Remove all consumed outputs
+	{
+		TArray<Chaos::TSimCallbackOutputHandle<FChaosSimModuleManagerAsyncOutput>> NewPendingOutputs;
+		for (int32 OutputIdx = LastOutputIdx; OutputIdx < PendingOutputs.Num(); ++OutputIdx)
+		{
+			NewPendingOutputs.Emplace(MoveTemp(PendingOutputs[OutputIdx]));
+		}
+		PendingOutputs = MoveTemp(NewPendingOutputs);
+	}
+
+	// It's possible we will end up multiple frames ahead of output, take the latest ready output.
+	Chaos::TSimCallbackOutputHandle<FChaosSimModuleManagerAsyncOutput> AsyncOutput;
+	Chaos::TSimCallbackOutputHandle<FChaosSimModuleManagerAsyncOutput> AsyncOutputLatest;
+	while ((AsyncOutputLatest = AsyncCallback->PopOutputData_External()))
+	{
+		AsyncOutput = MoveTemp(AsyncOutputLatest);
+	}
+
+	if (UWorld* World = Scene.GetOwningWorld())
+	{
+		{
+		int32 NumVehiclesInActiveBatch = 0;
+		for (TWeakObjectPtr<UModularVehicleComponent> Vehicle : GCVehicles)
+		{
+			auto NextOutput = PendingOutputs.Num() > 0 ? PendingOutputs[0].Get() : nullptr;
+			float Alpha = 0.f;
+			if (NextOutput && LatestOutput)
+			{
+				const float Denom = NextOutput->InternalTime - LatestOutput->InternalTime;
+				if (Denom > SMALL_NUMBER)
+				{
+					Alpha = (ResultsTime - LatestOutput->InternalTime) / Denom;
+				}
+			}
+
+			AsyncInput->VehicleInputs.Add(Vehicle->SetCurrentAsyncData(AsyncInput->VehicleInputs.Num(), LatestOutput.Get(), NextOutput, Alpha, Timestamp));
+		}
+
+		}
+		int32 NumVehiclesInActiveBatch = 0;
+
+		for (TWeakObjectPtr<UModularVehicleBaseComponent> Vehicle : CUVehicles)
+		{
+			auto NextOutput = PendingOutputs.Num() > 0 ? PendingOutputs[0].Get() : nullptr;
+			float Alpha = 0.f;
+			if (NextOutput && LatestOutput)
+			{
+				const float Denom = NextOutput->InternalTime - LatestOutput->InternalTime;
+				if (Denom > SMALL_NUMBER)
+				{
+					Alpha = (ResultsTime - LatestOutput->InternalTime) / Denom;
+				}
+			}
+
+			AsyncInput->VehicleInputs.Add(Vehicle->SetCurrentAsyncData(AsyncInput->VehicleInputs.Num(), LatestOutput.Get(), NextOutput, Alpha, Timestamp));
+		}
+
+	}
+
+	++Timestamp;
+
+	bool ForceSingleThread = !GSimModuleDebugParams.EnableMultithreading;
+
+
+	{
+		const auto& AwakeVehiclesBatch = GCVehicles;
+		auto LambdaParallelUpdate = [DeltaSeconds, &AwakeVehiclesBatch](int32 Idx)
+		{
+			TWeakObjectPtr<UModularVehicleComponent> Vehicle = AwakeVehiclesBatch[Idx];
+			Vehicle->ParallelUpdate(DeltaSeconds); // gets output state from PT
+		};
+
+		ParallelFor(AwakeVehiclesBatch.Num(), LambdaParallelUpdate, ForceSingleThread);
+	}
+
+
+	{
+		const auto& AwakeVehiclesBatch = CUVehicles;
+		auto LambdaParallelUpdate2 = [DeltaSeconds, &AwakeVehiclesBatch](int32 Idx)
+		{
+			TWeakObjectPtr<UModularVehicleBaseComponent> Vehicle = AwakeVehiclesBatch[Idx];
+			Vehicle->ParallelUpdate(DeltaSeconds); // gets output state from PT
+		};
+
+		ParallelFor(AwakeVehiclesBatch.Num(), LambdaParallelUpdate2, ForceSingleThread);
+	}
+
+}
+
diff --git a/Engine/Plugins/Experimental/ChaosModularVehicle/Source/ChaosModularVehicleEngine/Private/ChaosModularVehicle/ChaosSimModuleManagerAsyncCallback.cpp b/Engine/Plugins/Experimental/ChaosModularVehicle/Source/ChaosModularVehicleEngine/Private/ChaosModularVehicle/ChaosSimModuleManagerAsyncCallback.cpp
new file mode 100644
--- /dev/null
+++ b/Engine/Plugins/Experimental/ChaosModularVehicle/Source/ChaosModularVehicleEngine/Private/ChaosModularVehicle/ChaosSimModuleManagerAsyncCallback.cpp
@@ -0,0 +1,384 @@
+// Copyright Epic Games, Inc. All Rights Reserved.
+
+#include "ChaosModularVehicle/ChaosSimModuleManagerAsyncCallback.h"
+
+#include "ChaosModularVehicle/ModularVehicleComponent.h"
+#include "ChaosModularVehicle/ModularVehicleBaseComponent.h"
+#include "ChaosModularVehicle/ModularVehicleSimulationCU.h"
+#include "PBDRigidsSolver.h"
+#include "Chaos/ParticleHandleFwd.h"
+#include "PhysicsProxy/GeometryCollectionPhysicsProxy.h"
+#include "PhysicsProxy/ClusterUnionPhysicsProxy.h"
+
+FSimModuleDebugParams GSimModuleDebugParams;
+
+DECLARE_CYCLE_STAT(TEXT("AsyncCallback:OnPreSimulate_Internal"), STAT_AsyncCallback_OnPreSimulate, STATGROUP_ChaosSimModuleManager);
+
+FName FChaosSimModuleManagerAsyncCallback::GetFNameForStatId() const
+{
+	const static FLazyName StaticName("FChaosSimModuleManagerAsyncCallback");
+	return StaticName;
+}
+
+/**
+ * Callback from Physics thread
+ */
+
+void FChaosSimModuleManagerAsyncCallback::ProcessInputs_Internal(int32 PhysicsStep)
+{
+	const FChaosSimModuleManagerAsyncInput* AsyncInput = GetConsumerInput_Internal();
+	if (AsyncInput == nullptr)
+	{
+		return;
+	}
+
+	for (const TUniquePtr<FModularVehicleAsyncInput>& VehicleInput : AsyncInput->VehicleInputs)
+	{
+		VehicleInput->ProcessInputs();
+	}
+}
+
+/**
+ * Callback from Physics thread
+ */
+void FChaosSimModuleManagerAsyncCallback::OnPreSimulate_Internal()
+{
+	using namespace Chaos;
+
+	SCOPE_CYCLE_COUNTER(STAT_AsyncCallback_OnPreSimulate);
+
+	float DeltaTime = GetDeltaTime_Internal();
+	float SimTime = GetSimTime_Internal();
+
+	const FChaosSimModuleManagerAsyncInput* Input = GetConsumerInput_Internal();
+	if (Input == nullptr)
+	{
+		return;
+	}
+
+	const int32 NumVehicles = Input->VehicleInputs.Num();
+
+	UWorld* World = Input->World.Get();	//only safe to access for scene queries
+	if (World == nullptr || NumVehicles == 0)
+	{
+		//world is gone so don't bother, or nothing to simulate.
+		return;
+	}
+
+	Chaos::FPhysicsSolver* PhysicsSolver = static_cast<Chaos::FPhysicsSolver*>(GetSolver());
+	if (PhysicsSolver == nullptr)
+	{
+		return;
+	}
+
+	FChaosSimModuleManagerAsyncOutput& Output = GetProducerOutputData_Internal();
+	Output.VehicleOutputs.AddDefaulted(NumVehicles);
+	Output.Timestamp = Input->Timestamp;
+
+	const TArray<TUniquePtr<FModularVehicleAsyncInput>>& InputVehiclesBatch = Input->VehicleInputs;
+	TArray<TUniquePtr<FModularVehicleAsyncOutput>>& OutputVehiclesBatch = Output.VehicleOutputs;
+
+	// beware running the vehicle simulation in parallel, code must remain threadsafe
+	auto LambdaParallelUpdate = [World, DeltaTime, SimTime, &InputVehiclesBatch, &OutputVehiclesBatch](int32 Idx)
+	{
+		const FModularVehicleAsyncInput& VehicleInput = *InputVehiclesBatch[Idx];
+
+		if (VehicleInput.Proxy == nullptr)
+		{
+			return;
+		}
+
+		bool bWake = false;
+		OutputVehiclesBatch[Idx] = VehicleInput.Simulate(World, DeltaTime, SimTime, bWake);
+
+	};
+
+	bool ForceSingleThread = !GSimModuleDebugParams.EnableMultithreading;
+	PhysicsParallelFor(OutputVehiclesBatch.Num(), LambdaParallelUpdate, ForceSingleThread);
+
+	// Delayed application of forces - This is separate from Simulate because forces cannot be executed multi-threaded
+	for (const TUniquePtr<FModularVehicleAsyncInput>& VehicleInput : InputVehiclesBatch)
+	{
+		if (VehicleInput.IsValid())
+		{
+			VehicleInput->ApplyDeferredForces();
+		}
+	}
+}
+
+/**
+ * Contact modification currently unused
+ */
+void FChaosSimModuleManagerAsyncCallback::OnContactModification_Internal(Chaos::FCollisionContactModifier& Modifications)
+{
+
+}
+
+
+TUniquePtr<FModularVehicleAsyncOutput> FModularVehicleAsyncInput::Simulate(UWorld* World, const float DeltaSeconds, const float TotalSeconds, bool& bWakeOut) const
+{
+	TUniquePtr<FModularVehicleAsyncOutput> Output = MakeUnique<FModularVehicleAsyncOutput>();
+
+	//support nullptr because it allows us to go wide on filling the async inputs
+	if (Proxy == nullptr)
+	{
+		return Output;
+	}
+
+	if (Vehicle && Vehicle->VehicleSimulationPT)
+	{
+		// FILL OUTPUT DATA HERE THAT WILL GET PASSED BACK TO THE GAME THREAD
+		Vehicle->VehicleSimulationPT->Simulate(World, DeltaSeconds, *this, *Output.Get(), Proxy);
+
+		FModularVehicleAsyncOutput& OutputData = *Output.Get();
+		Vehicle->VehicleSimulationPT->FillOutputState(OutputData);
+	}
+
+
+	Output->bValid = true;
+
+	return MoveTemp(Output);
+}
+
+void FModularVehicleAsyncInput::ApplyDeferredForces() const
+{
+	if (Vehicle && Proxy)
+	{
+		if (Proxy->GetType() == EPhysicsProxyType::ClusterUnionProxy)
+		{
+			Vehicle->VehicleSimulationPT->ApplyDeferredForces(static_cast<Chaos::FClusterUnionPhysicsProxy*>(Proxy));
+		}
+		else if (Proxy->GetType() == EPhysicsProxyType::GeometryCollectionType)
+		{
+			Vehicle->VehicleSimulationPT->ApplyDeferredForces(static_cast<FGeometryCollectionPhysicsProxy*>(Proxy));
+		}
+
+	}
+
+}
+
+void FModularVehicleAsyncInput::ProcessInputs()
+{
+	if (!GetVehicle())
+	{
+		return;
+	}
+
+	FModularVehicleSimulationCU* VehicleSim = GetVehicle()->VehicleSimulationPT.Get();
+
+	if (VehicleSim == nullptr || !GetVehicle()->bUsingNetworkPhysicsPrediction || GetVehicle()->GetWorld() == nullptr)
+	{
+		return;
+	}
+	bool bIsResimming = false;
+	if (FPhysScene* PhysScene = GetVehicle()->GetWorld()->GetPhysicsScene())
+	{
+		if (Chaos::FPhysicsSolver* LocalSolver = PhysScene->GetSolver())
+		{
+			bIsResimming = LocalSolver->GetEvolution()->IsResimming();
+		}
+	}
+
+	APlayerController* PlayerController = GetVehicle()->GetPlayerController();
+	if (PlayerController && PlayerController->IsLocalController() && !bIsResimming)
+	{
+		VehicleSim->VehicleInputs = PhysicsInputs.NetworkInputs.VehicleInputs;
+	}
+	else
+	{
+		PhysicsInputs.NetworkInputs.VehicleInputs = VehicleSim->VehicleInputs;
+	}
+
+}
+
+bool FNetworkModularVehicleInputs::NetSerialize(FArchive& Ar, class UPackageMap* Map, bool& bOutSuccess)
+{
+	FNetworkPhysicsDatas::SerializeFrames(Ar);
+
+	Ar << VehicleInputs.Steering;
+	Ar << VehicleInputs.Throttle;
+	Ar << VehicleInputs.Brake;
+	Ar << VehicleInputs.Handbrake;
+	Ar << VehicleInputs.Pitch;
+	Ar << VehicleInputs.Roll;
+	Ar << VehicleInputs.Yaw;
+	Ar << VehicleInputs.Boost;
+	Ar << VehicleInputs.Drift;
+	Ar << VehicleInputs.Reverse;
+	Ar << VehicleInputs.KeepAwake;
+
+	bOutSuccess = true;
+	return bOutSuccess;
+}
+
+void FNetworkModularVehicleInputs::ApplyDatas(UActorComponent* NetworkComponent) const
+{
+	if (GSimModuleDebugParams.EnableNetworkStateData)
+	{
+		if (FModularVehicleSimulationCU* VehicleSimulation = Cast<UModularVehicleBaseComponent>(NetworkComponent)->VehicleSimulationPT.Get())
+		{
+			VehicleSimulation->VehicleInputs = VehicleInputs;
+		}
+	}
+}
+
+void FNetworkModularVehicleInputs::BuildDatas(const UActorComponent* NetworkComponent)
+{
+	if (GSimModuleDebugParams.EnableNetworkStateData && NetworkComponent)
+	{
+		if (const FModularVehicleSimulationCU* VehicleSimulation = Cast<const UModularVehicleBaseComponent>(NetworkComponent)->VehicleSimulationPT.Get())
+		{
+			VehicleInputs = VehicleSimulation->VehicleInputs;
+		}
+	}
+}
+
+void FNetworkModularVehicleInputs::InterpolateDatas(const FNetworkModularVehicleInputs& MinDatas, const FNetworkModularVehicleInputs& MaxDatas)
+{
+	const float LerpFactor = (LocalFrame - MinDatas.LocalFrame) / (MaxDatas.LocalFrame - MinDatas.LocalFrame);
+
+	VehicleInputs.Steering = FMath::Lerp(MinDatas.VehicleInputs.Steering, MaxDatas.VehicleInputs.Steering, LerpFactor);
+	VehicleInputs.Throttle = FMath::Lerp(MinDatas.VehicleInputs.Throttle, MaxDatas.VehicleInputs.Throttle, LerpFactor);
+	VehicleInputs.Brake = FMath::Lerp(MinDatas.VehicleInputs.Brake, MaxDatas.VehicleInputs.Brake, LerpFactor);
+	VehicleInputs.Handbrake = FMath::Lerp(MinDatas.VehicleInputs.Handbrake, MaxDatas.VehicleInputs.Handbrake, LerpFactor);
+	VehicleInputs.Pitch = FMath::Lerp(MinDatas.VehicleInputs.Pitch, MaxDatas.VehicleInputs.Pitch, LerpFactor);
+	VehicleInputs.Roll = FMath::Lerp(MinDatas.VehicleInputs.Roll, MaxDatas.VehicleInputs.Roll, LerpFactor);
+	VehicleInputs.Yaw = FMath::Lerp(MinDatas.VehicleInputs.Yaw, MaxDatas.VehicleInputs.Yaw, LerpFactor);
+	VehicleInputs.Boost = FMath::Lerp(MinDatas.VehicleInputs.Boost, MaxDatas.VehicleInputs.Boost, LerpFactor);
+	VehicleInputs.Drift = FMath::Lerp(MinDatas.VehicleInputs.Drift, MaxDatas.VehicleInputs.Drift, LerpFactor);
+	VehicleInputs.Reverse = MinDatas.VehicleInputs.Reverse;
+	VehicleInputs.KeepAwake = MinDatas.VehicleInputs.KeepAwake;
+}
+
+bool FNetworkModularVehicleStates::NetSerialize(FArchive& Ar, class UPackageMap* Map, bool& bOutSuccess)
+{
+	FNetworkPhysicsDatas::SerializeFrames(Ar);
+
+	int32 NumNetModules = ModuleData.Num();
+	Ar << NumNetModules;
+
+	for (int I = 0; I < NumNetModules; I++)
+	{
+		if (Ar.IsLoading())
+		{
+			if (NumNetModules > 0)
+			{
+				int32 ModuleType = Chaos::eSimType::Undefined;
+				int32 SimArrayIndex = 0;
+				Ar << ModuleType;
+				Ar << SimArrayIndex;
+
+				if (!ModuleData.IsEmpty())
+				{
+					ensure(I <= ModuleData.Num());
+				}
+
+				if (ModuleData.Num() != NumNetModules)
+				{
+					ModuleData.Reserve(NumNetModules);
+					switch (ModuleType)
+					{
+					case Chaos::eSimType::Suspension:
+					{
+						ModuleData.Emplace(MakeShared<Chaos::FSuspensionSimModuleDatas>(SimArrayIndex
+#if !(UE_BUILD_SHIPPING || UE_BUILD_TEST)
+							, FString()
+#endif
+						));
+					}
+					break;
+
+					case Chaos::eSimType::Transmission:
+					{
+						ModuleData.Emplace(MakeShared<Chaos::FTransmissionSimModuleDatas>(SimArrayIndex
+#if !(UE_BUILD_SHIPPING || UE_BUILD_TEST)
+							, FString()
+#endif
+						));
+					}
+					break;
+
+					case Chaos::eSimType::Engine:
+					{
+						ModuleData.Emplace(MakeShared<Chaos::FEngineSimModuleDatas>(SimArrayIndex
+#if !(UE_BUILD_SHIPPING || UE_BUILD_TEST)
+							, FString()
+#endif
+						));
+					}
+					break;
+
+					case Chaos::eSimType::Clutch:
+					{
+						ModuleData.Emplace(MakeShared<Chaos::FClutchSimModuleDatas>(SimArrayIndex
+#if !(UE_BUILD_SHIPPING || UE_BUILD_TEST)
+							, FString()
+#endif
+						));
+					}
+					break;
+
+					case Chaos::eSimType::Wheel:
+					{
+						ModuleData.Emplace(MakeShared<Chaos::FWheelSimModuleDatas>(SimArrayIndex
+#if !(UE_BUILD_SHIPPING || UE_BUILD_TEST)
+							, FString()
+#endif
+						));
+					}
+					break;
+
+					default:
+					{
+						checkf(false, TEXT("Unhandled NetModuleType case"));
+					}
+					}
+				}
+			}
+		}
+		else
+		{
+			int32 ModuleType = (int32)ModuleData[I]->GetType();
+			Ar << ModuleType;
+			Ar << ModuleData[I]->SimArrayIndex;
+		}
+
+		ModuleData[I]->Serialize(Ar);
+	}
+
+	return true;
+}
+
+void FNetworkModularVehicleStates::ApplyDatas(UActorComponent* NetworkComponent) const
+{
+	if (FModularVehicleSimulationCU* VehicleSimulation = Cast<UModularVehicleBaseComponent>(NetworkComponent)->VehicleSimulationPT.Get())
+	{
+		VehicleSimulation->AccessSimComponentTree()->SetSimState(ModuleData);
+	}
+}
+
+void FNetworkModularVehicleStates::BuildDatas(const UActorComponent* NetworkComponent)
+{
+	if (NetworkComponent)
+	{
+		if (const FModularVehicleSimulationCU* VehicleSimulation = Cast<const UModularVehicleBaseComponent>(NetworkComponent)->VehicleSimulationPT.Get())
+		{
+			VehicleSimulation->GetSimComponentTree()->SetNetState(ModuleData);
+		}
+	}
+}
+
+void FNetworkModularVehicleStates::InterpolateDatas(const FNetworkModularVehicleStates& MinDatas, const FNetworkModularVehicleStates& MaxDatas)
+{
+	const float LerpFactor = (LocalFrame - MinDatas.LocalFrame) / (MaxDatas.LocalFrame - MinDatas.LocalFrame);
+
+	for (int I = 0; I < ModuleData.Num(); I++)
+	{
+		// if these don't match then something has gone terribly wrong
+		check(ModuleData[I]->GetType() == MinDatas.ModuleData[I]->GetType());
+		check(ModuleData[I]->GetType() == MaxDatas.ModuleData[I]->GetType());
+
+		ModuleData[I]->Lerp(LerpFactor, *MinDatas.ModuleData[I].Get(), *MaxDatas.ModuleData[I].Get());
+	}
+}
+
diff --git a/Engine/Plugins/Experimental/ChaosModularVehicle/Source/ChaosModularVehicleEngine/Private/ChaosModularVehicle/ModularVehicleActor.cpp b/Engine/Plugins/Experimental/ChaosModularVehicle/Source/ChaosModularVehicleEngine/Private/ChaosModularVehicle/ModularVehicleActor.cpp
new file mode 100644
--- /dev/null
+++ b/Engine/Plugins/Experimental/ChaosModularVehicle/Source/ChaosModularVehicleEngine/Private/ChaosModularVehicle/ModularVehicleActor.cpp
@@ -0,0 +1,40 @@
+// Copyright Epic Games, Inc. All Rights Reserved.
+
+#include "ChaosModularVehicle/ModularVehicleActor.h"
+
+DEFINE_LOG_CATEGORY_STATIC(AModularVehicleLogging, Log, All);
+
+AModularVehicleActor::AModularVehicleActor(const FObjectInitializer& ObjectInitializer)
+	: Super(ObjectInitializer)
+{
+	ModularVehicleComponent = CreateDefaultSubobject<UModularVehicleComponent>(TEXT("ModularVehicleComponent0"));
+
+	RootComponent = ModularVehicleComponent;
+
+	PrimaryActorTick.bCanEverTick = true;
+	SetActorTickEnabled(true);
+
+	bReplicates = true;
+}
+
+void AModularVehicleActor::PostInitializeComponents()
+{
+	Super::PostInitializeComponents();
+
+	ModularVehicleComponent->CreateVehicleSim();
+}
+
+void AModularVehicleActor::TickSimulation(float DeltaTime)
+{
+}
+
+void AModularVehicleActor::Tick(float DeltaTime) 
+{
+	Super::Tick(DeltaTime);
+
+	if (ModularVehicleComponent)
+	{
+		ModularVehicleComponent->SetRenderStateDirty();
+	}
+}
+
