Move Modular Vehicle Plugin files out of restricted folders

diff --git a/Engine/Plugins/Experimental/ChaosModularVehicle/Source/ChaosModularVehicleEngine/Private/ChaosModularVehicle/ModularVehicleBaseComponent.cpp b/Engine/Plugins/Experimental/ChaosModularVehicle/Source/ChaosModularVehicleEngine/Private/ChaosModularVehicle/ModularVehicleBaseComponent.cpp
new file mode 100644
--- /dev/null
+++ b/Engine/Plugins/Experimental/ChaosModularVehicle/Source/ChaosModularVehicleEngine/Private/ChaosModularVehicle/ModularVehicleBaseComponent.cpp
@@ -0,0 +1,1125 @@
+// Copyright Epic Games, Inc. All Rights Reserved.
+
+#include "ChaosModularVehicle/ModularVehicleBaseComponent.h"
+#include "ChaosModularVehicle/ChaosSimModuleManager.h"
+#include "ChaosModularVehicle/VehicleSimBaseComponent.h"
+#include "ChaosModularVehicle/ModularVehicleDefaultAsyncInput.h"
+#include "Engine/Engine.h"
+#include "Engine/Canvas.h"
+#include "PhysicsEngine/PhysicsSettings.h"
+#include "PhysicsEngine/ClusterUnionComponent.h"
+#include "PhysicsEngine/PhysicsObjectExternalInterface.h"
+#include "PhysicsReplication.h"
+#include "Net/Core/PushModel/PushModel.h"
+#include "Net/UnrealNetwork.h"
+#include "GameFramework/PlayerController.h"
+#include "GameFramework/Pawn.h"
+#include "Physics/Experimental/PhysScene_Chaos.h"
+#include "SimModule/SimModuleTree.h"
+#include "Chaos/PBDSuspensionConstraints.h"
+
+#include "ChaosModularVehicle/ChaosSimModuleManagerAsyncCallback.h"
+
+DEFINE_LOG_CATEGORY(LogModularBase);
+
+bool bModularVehicle_SuspensionConstraint_Enabled = true;
+FAutoConsoleVariableRef CVarModularVehicleSuspensionConstraintEnabled(TEXT("p.ModularVehicle.SuspensionConstraint.Enabled"), bModularVehicle_SuspensionConstraint_Enabled, TEXT("Enable/Disable suspension constraint falling back to simple forces when constraint disabled (requires restart)."));
+
+bool bModularVehicle_DumpModuleTreeStructure_Enabled = false;
+FAutoConsoleVariableRef CVarModularVehicleDumpModuleTreeStructureEnabled(TEXT("p.ModularVehicle.DumpModuleTreeStructure.Enabled"), bModularVehicle_DumpModuleTreeStructure_Enabled, TEXT("Enable/Disable logging of module tree structure every time there is a change."));
+
+UModularVehicleBaseComponent::UModularVehicleBaseComponent(const FObjectInitializer& ObjectInitializer)
+	: Super(ObjectInitializer)
+{
+	PrimaryComponentTick.bCanEverTick = true;
+	bTickInEditor = false;
+	bRequiresControllerForInputs = true;
+
+	LinearDamping = 0.01f;
+	AngularDamping = 0.0f;
+
+	bAutoAddComponentsFromWorld = false;
+	AutoAddOverlappingBoxSize = FVector(500, 500, 500);
+	ClusteringCount = 0;
+	DelayClusteringCount = 0;
+
+	SetIsReplicatedByDefault(true);
+	bUsingNetworkPhysicsPrediction = Chaos::FPhysicsSolverBase::IsNetworkPhysicsPredictionEnabled();
+
+	SuspensionTraceCollisionResponses = FCollisionResponseContainer::GetDefaultResponseContainer();
+	SuspensionTraceCollisionResponses.Vehicle = ECR_Ignore;
+	bSuspensionTraceComplex = true;
+	bKeepVehicleAwake = true;
+
+	CurrentGear = 0;
+	EngineRPM = 0.0f;
+	EngineTorque = 0.0f;
+
+	//// #TODO: currently ordering of this must match EModularVehicleInputType
+	//InputInterpolationRates.Reset();
+	//InputInterpolationRates.Add(FModularVehicleInputRate(FString("Throttle")));
+	//InputInterpolationRates.Add(FModularVehicleInputRate(FString("Brake")));
+	//InputInterpolationRates.Add(FModularVehicleInputRate(FString("Clutch")));
+	//InputInterpolationRates.Add(FModularVehicleInputRate(FString("Steering")));
+	//InputInterpolationRates.Add(FModularVehicleInputRate(FString("Handbrake")));
+	//InputInterpolationRates.Add(FModularVehicleInputRate(FString("Pitch")));
+	//InputInterpolationRates.Add(FModularVehicleInputRate(FString("Roll")));
+	//InputInterpolationRates.Add(FModularVehicleInputRate(FString("Yaw")));
+	//InputInterpolationRates.Add(FModularVehicleInputRate(FString("Gear")));
+	//InputInterpolationRates.Add(FModularVehicleInputRate(FString("DebugIndex")));
+
+	if (bUsingNetworkPhysicsPrediction)
+	{
+		static const FName NetworkPhysicsComponentName(TEXT("PC_NetworkPhysicsComponent"));
+
+		NetworkPhysicsComponent = CreateDefaultSubobject<UNetworkPhysicsComponent, UNetworkPhysicsComponent>(NetworkPhysicsComponentName);
+		NetworkPhysicsComponent->SetNetAddressable();
+		NetworkPhysicsComponent->SetIsReplicated(true);
+	}
+
+}
+
+UModularVehicleBaseComponent::~UModularVehicleBaseComponent()
+{
+}
+
+APlayerController* UModularVehicleBaseComponent::GetPlayerController() const
+{
+	if (AController* Controller = GetController())
+	{
+		return Cast<APlayerController>(Controller);
+	}
+
+	if (APawn* Pawn = Cast<APawn>(GetOwner()))
+	{
+		if (AController* Controller = Pawn->GetController())
+		{
+			return Cast<APlayerController>(Controller);
+		}
+
+		if (APlayerController* PC = Cast<APlayerController>(Pawn->GetOwner()))
+		{
+			return PC;
+		}
+
+	}
+	return nullptr;
+}
+
+void UModularVehicleBaseComponent::OnCreatePhysicsState()
+{
+	Super::OnCreatePhysicsState();
+
+	if (ClusterUnionComponent)
+	{
+		// piggyback on the Add/Remove component events from the cluster union to add/remove simulation modules
+		ClusterUnionComponent->OnComponentAddedNativeEvent.AddUObject(this, &UModularVehicleBaseComponent::AddComponentToSimulation);
+		ClusterUnionComponent->OnComponentRemovedNativeEvent.AddUObject(this, &UModularVehicleBaseComponent::RemoveComponentFromSimulation);
+
+		// #TODO: this doesn't appear to be working
+		if (bKeepVehicleAwake && ClusterUnionComponent->GetPhysicsProxy())
+		{
+			if (auto* Particle = ClusterUnionComponent->GetPhysicsProxy()->GetParticle_External())
+			{
+				Particle->SetSleepType(Chaos::ESleepType::NeverSleep);
+			}
+		}
+	}
+
+	// convert the game thread side UVehicleSimBaseComponent(s) to physics thread simulation ISimulationModuleBase(s)
+	CreateVehicleSim();
+
+	if (bUsingNetworkPhysicsPrediction)
+	{
+		if (NetworkPhysicsComponent)
+		{
+			// register interface to handle network prediction callbacks
+			// #Note: in our case we don't yet know what the replication data will be since the modules are built after this point at runtime
+			NetworkPhysicsComponent->CreateDatasHistory<FPhysicsModularVehicleTraits>(this);
+		}
+	}
+
+	NextConstructionIndex = 0;
+
+	ActorsToIgnore.Add(GetOwner()); // vehicle ignore self in scene query
+
+}
+
+void UModularVehicleBaseComponent::OnDestroyPhysicsState()
+{
+	Super::OnDestroyPhysicsState();
+
+	if (ClusterUnionComponent)
+	{
+		ClusterUnionComponent->OnComponentAddedNativeEvent.RemoveAll(this);
+		ClusterUnionComponent->OnComponentRemovedNativeEvent.RemoveAll(this);
+	}
+
+	DestroyVehicleSim();
+
+	if (bUsingNetworkPhysicsPrediction && NetworkPhysicsComponent)
+	{
+		NetworkPhysicsComponent->RemoveDatasHistory();
+	}
+
+}
+
+
+int GenerateNewGuid()
+{
+	static int Val = 0;
+	return Val++;
+}
+
+void UModularVehicleBaseComponent::CreateAssociatedSimComponents(const UPrimitiveComponent* AttachedComponent, int ParentIndex, int TransformIndex, Chaos::FSimTreeUpdates& TreeUpdatesOut)
+{
+	using namespace Chaos;
+	if (AttachedComponent == nullptr || ClusterUnionComponent == nullptr)
+	{
+		return;
+	}
+
+	UE_LOG(LogModularBase, Log, TEXT("CreateAssociatedSimComponents Attaching %s, TransformIndex %d"), *AttachedComponent->GetName(), TransformIndex);
+
+	TArray<Chaos::FClusterUnionChildData> ChildParticles = ClusterUnionComponent->GetPhysicsProxy()->GetSyncedData_External().ChildParticles;
+
+	ensure(TransformIndex < ChildParticles.Num());
+
+	if (const UVehicleSimBaseComponent* Component = Cast<UVehicleSimBaseComponent>(AttachedComponent))
+	{
+		FTransform ComponentTransform = AttachedComponent->GetComponentTransform();
+		ComponentTransform.SetScale3D(FVector(1.0f, 1.0f, 1.0f));
+		ComponentTransform.SetTranslation(FVector::ZeroVector);
+
+		int TreeIndex = INDEX_NONE;
+
+		Chaos::ISimulationModuleBase* NewModule = Component->CreateNewCoreModule();
+
+		TUniquePtr<Chaos::FSimModuleTree>& SimModuleTree = VehicleSimulationPT->AccessSimComponentTree();
+
+#if !(UE_BUILD_SHIPPING || UE_BUILD_TEST)
+		FString DebugString;
+		NewModule->GetDebugString(DebugString);
+		UE_LOG(LogModularBase, Log, TEXT("CreateAssociatedSimComponents Found Sim Component %s (%s)"), *NewModule->GetDebugName(), *DebugString);
+#endif
+
+		FVector LocalOffset(0, 0, 0);
+		FVector LocalOffsetCOM(0, 0, 0);
+		FTransform PhysicsTransform = FTransform::Identity;
+
+		FTransform InitialTransform = ChildParticles[TransformIndex].ChildToParent;
+		InitialTransform.SetLocation(InitialTransform.GetLocation());
+		NewModule->SetSimModuleTree(SimModuleTree.Get());
+		TreeIndex = TreeUpdatesOut.AddNodeBelow(ParentIndex, NewModule);
+		NewModule->SetTransformIndex(TransformIndex);
+		NewModule->SetParticleIndex(ChildParticles[TransformIndex].ParticleIdx);
+
+		NewModule->SetIntactTransform(FTransform::Identity);
+		NewModule->SetClusteredTransform(InitialTransform);
+		NewModule->SetClustered(true);
+
+		FVehicleComponentData& ComponentData = ComponentToPhysicsObjects.FindOrAdd(AttachedComponent);
+		ComponentData.Guid = GenerateNewGuid();
+		NewModule->SetGuid(ComponentData.Guid);
+
+		NewModule->SetInitialParticleTransform(InitialTransform);
+		NewModule->SetComponentTransform(ComponentTransform);
+
+		FTransform ClusterredTransform(FQuat::Identity, InitialTransform.GetLocation());
+		NewModule->SetClusteredTransform(ClusterredTransform);
+
+		ParentIndex = TreeIndex;
+
+		if (bModularVehicle_SuspensionConstraint_Enabled)
+		{
+			if (NewModule->GetSimType() == Chaos::eSimType::Suspension)
+			{
+				CreateConstraint(NewModule);
+			}
+		}
+	}
+
+	// get this components children, not all children
+	TArray<USceneComponent*> Children;
+	AttachedComponent->GetChildrenComponents(false, Children);
+
+	// recurse down tree, converting all SimComponents to proper simulation modules
+	for (USceneComponent* Child : Children)
+	{
+		if (UVehicleSimBaseComponent* ChildSimComponent = Cast<UVehicleSimBaseComponent>(Child))
+		{
+			CreateAssociatedSimComponents(ChildSimComponent, ParentIndex, TransformIndex, TreeUpdatesOut);
+		}
+	}
+
+	UpdatePhysicalProperties();
+}
+
+void UModularVehicleBaseComponent::UpdatePhysicalProperties()
+{
+	if (Chaos::FClusterUnionPhysicsProxy* Proxy = ClusterUnionComponent->GetPhysicsProxy())
+	{
+		Chaos::FPBDRigidsSolver* Solver = Proxy->GetSolver<Chaos::FPBDRigidsSolver>();
+		Solver->EnqueueCommandImmediate([Proxy, this]() mutable
+			{
+				Proxy->GetParticle_Internal()->SetLinearEtherDrag(LinearDamping);
+				Proxy->GetParticle_Internal()->SetAngularEtherDrag(AngularDamping);
+			});
+
+		// updating external particle currently doesn't update physics particle
+		//	if (auto* Particle = static_cast<Chaos::FClusterUnionPhysicsProxy*>(Proxy)->GetParticle_External())
+		//	{
+		//		Particle->SetLinearEtherDrag(LinearDamping);
+		//		Particle->SetAngularEtherDrag(AngularDamping);
+		//	}
+	}
+}
+
+void UModularVehicleBaseComponent::BeginPlay()
+{
+	Super::BeginPlay();
+
+	const bool bHasAuthority = GetOwner()->HasAuthority();
+	if (bHasAuthority)
+	{
+		if (bAutoAddComponentsFromWorld && (DelayClusteringCount == 0))
+		{
+			AddOverlappingComponentsToCluster();
+		}
+		else
+		{
+			AddGeometryCollectionsFromOwnedActor();
+		}
+	}
+}
+
+void UModularVehicleBaseComponent::TickComponent(float DeltaTime, ELevelTick TickType, FActorComponentTickFunction* ThisTickFunction)
+{
+	Super::TickComponent(DeltaTime, TickType, ThisTickFunction);
+
+	const bool bHasAuthority = GetOwner()->HasAuthority();
+	if (bAutoAddComponentsFromWorld && bHasAuthority && (++ClusteringCount == DelayClusteringCount))
+	{
+		AddOverlappingComponentsToCluster();
+	}
+
+}
+
+void UModularVehicleBaseComponent::CreateConstraint(Chaos::ISimulationModuleBase* NewModule)
+{
+	check(NewModule->GetSimType() == Chaos::eSimType::Suspension);
+	Chaos::FSuspensionSimModule* SuspensionModule = static_cast<Chaos::FSuspensionSimModule*>(NewModule);
+
+	const Chaos::FSuspensionSettings& Setup = SuspensionModule->Setup();
+	const FVector& LocalOffset = SuspensionModule->GetInitialParticleTransform().GetLocation();
+
+	if (ClusterUnionComponent && ClusterUnionComponent->GetPhysicsProxy())
+	{
+		const Chaos::FPhysicsObjectHandle& PhysicsObject = ClusterUnionComponent->GetPhysicsProxy()->GetPhysicsObjectHandle();
+
+		if ( FChaosScene* Scene = static_cast<FChaosScene*>(FPhysicsObjectExternalInterface::GetScene({ &PhysicsObject, 1 })) )
+		{
+			FLockedWritePhysicsObjectExternalInterface Interface = FPhysicsObjectExternalInterface::LockWrite(Scene);
+			if (const Chaos::FGeometryParticle* Particle = Interface->GetParticle(PhysicsObject))
+			{
+				FPhysicsConstraintHandle ConstraintHandle = FPhysicsInterface::CreateSuspension(PhysicsObject, LocalOffset);
+
+				if (ConstraintHandle.IsValid())
+				{
+					ConstraintHandles.Add(ConstraintHandle);
+
+					if (Chaos::FSuspensionConstraint* Constraint = static_cast<Chaos::FSuspensionConstraint*>(ConstraintHandle.Constraint))
+					{
+						Constraint->SetHardstopStiffness(1.0f);
+						Constraint->SetSpringStiffness(Setup.SpringRate * 0.25f);
+						Constraint->SetSpringPreload(Setup.SpringPreload);
+						Constraint->SetSpringDamping(Setup.SpringDamping * 5.0f);
+						Constraint->SetMinLength(-Setup.MaxRaise);
+						Constraint->SetMaxLength(Setup.MaxDrop);
+						Constraint->SetAxis(-Setup.SuspensionAxis);
+
+						SuspensionModule->SetSuspensionConstraint(Constraint);
+						SuspensionModule->SetConstraintIndex(ConstraintHandles.Num()-1);
+					}
+				}
+			}
+		}
+	}
+}
+
+void UModularVehicleBaseComponent::DestroyConstraint(int ConstraintIndex)
+{
+	if (ConstraintIndex >= 0 && ConstraintIndex < ConstraintHandles.Num())
+	{
+		FPhysicsConstraintHandle ConstraintHandle = ConstraintHandles[ConstraintIndex];
+		FPhysicsCommand::ExecuteWrite(ConstraintHandle, [&](const FPhysicsConstraintHandle& Constraint)
+			{
+				FPhysicsInterface::ReleaseConstraint(ConstraintHandle);
+			});
+
+		ConstraintHandles[ConstraintIndex].Reset();
+	}
+
+}
+
+void UModularVehicleBaseComponent::DestroyAllConstraints()
+{
+	if (ConstraintHandles.Num() > 0)
+	{
+		for (FPhysicsConstraintHandle ConstraintHandle : ConstraintHandles)
+		{
+			if (ConstraintHandle.IsValid())
+			{
+				FPhysicsCommand::ExecuteWrite(ConstraintHandle, [&](const FPhysicsConstraintHandle& Constraint)
+					{
+						FPhysicsInterface::ReleaseConstraint(ConstraintHandle);
+					});
+			}
+		}
+	}
+	ConstraintHandles.Empty();
+
+}
+
+void UModularVehicleBaseComponent::EnableConstraint(int ConstraintIndex, bool bEnabled)
+{
+	if (ConstraintIndex >= 0 && ConstraintIndex < ConstraintHandles.Num())
+	{
+		FPhysicsConstraintHandle ConstraintHandle = ConstraintHandles[ConstraintIndex];
+		FPhysicsCommand::ExecuteWrite(ConstraintHandle, [&](const FPhysicsConstraintHandle& Constraint)
+			{
+				if (Chaos::FSuspensionConstraint* SuspensionConstraint = static_cast<Chaos::FSuspensionConstraint*>(ConstraintHandle.Constraint))
+				{
+					SuspensionConstraint->SetEnabled(bEnabled);
+				}
+			});
+
+		ConstraintHandles[ConstraintIndex].Reset();
+	}
+
+}
+
+int32 UModularVehicleBaseComponent::FindComponentAddOrder(UPrimitiveComponent* InComponent)
+{
+	int32 ComponentAddOrder = INDEX_NONE;
+	for (const FConstructionData& Data : ConstructionDatas)
+	{
+		if (Data.Component == InComponent)
+		{
+			ComponentAddOrder = Data.ConstructionIndex;
+			break;
+		}
+	}
+	return ComponentAddOrder;
+}
+
+bool UModularVehicleBaseComponent::FindAndRemoveNextPendingUpdate(int32 NextIndex, Chaos::FSimTreeUpdates* OutData)
+{
+	bool bFound = false;
+	for (const TPair<TObjectKey<UPrimitiveComponent>, Chaos::FSimTreeUpdates>& Kvp : PendingTreeUpdates)
+	{
+		if (FindComponentAddOrder(Kvp.Key.ResolveObjectPtr()) == NextIndex)
+		{
+			*OutData = PendingTreeUpdates.FindAndRemoveChecked(Kvp.Key.ResolveObjectPtr());
+			bFound = true;
+			break;
+		}
+	}
+
+	return bFound;
+}
+
+void UModularVehicleBaseComponent::AddActorsToIgnore(TArray<AActor*>& ActorsIn)
+{
+	for (AActor* Actor : ActorsIn)
+	{
+		ActorsToIgnore.AddUnique(Actor);
+	}
+}
+
+void UModularVehicleBaseComponent::RemoveActorsToIgnore(TArray<AActor*>& ActorsIn)
+{
+	for (AActor* Actor : ActorsIn)
+	{
+		ActorsToIgnore.Remove(Actor);
+	}
+}
+
+void UModularVehicleBaseComponent::PreTickGT(float DeltaTime)
+{
+	if (PendingTreeUpdates.Num() > 0)
+	{
+		Chaos::FSimTreeUpdates Updates;
+		while (FindAndRemoveNextPendingUpdate(LastComponentAddIndex + 1, &Updates))
+		{
+			ActionTreeUpdates(&Updates);
+			LastComponentAddIndex++;
+		}
+	}
+
+	if (PVehicleOutput && UpdatedComponent)
+	{
+		APawn* MyOwner = Cast<APawn>(UpdatedComponent->GetOwner());
+		if (MyOwner)
+		{
+			// process control inputs and other data
+			UpdateState(DeltaTime);
+		}
+	}
+
+}
+
+void UModularVehicleBaseComponent::UpdateState(float DeltaTime)
+{
+	// update input values
+	AController* Controller = GetPlayerController();
+
+	bool bProcessLocally = bRequiresControllerForInputs ? (Controller && Controller->IsLocalController()) : true;
+
+	// IsLocallyControlled will fail if the owner is unpossessed (i.e. Controller == nullptr);
+	// Should we remove input instead of relying on replicated state in that case?
+	if (bProcessLocally && PVehicleOutput)
+	{
+		// Apply Inputs locally
+		SteeringInput = SteeringInputRate.InterpInputValue(DeltaTime, SteeringInput, RawSteeringInput);
+		ThrottleInput = ThrottleInputRate.InterpInputValue(DeltaTime, ThrottleInput, RawThrottleInput);
+		BrakeInput = BrakeInputRate.InterpInputValue(DeltaTime, BrakeInput, RawBrakeInput);
+		HandbrakeInput = HandbrakeInputRate.InterpInputValue(DeltaTime, HandbrakeInput, RawHandbrakeInput);
+		PitchInput = PitchInputRate.InterpInputValue(DeltaTime, PitchInput, RawPitchInput);
+		RollInput = RollInputRate.InterpInputValue(DeltaTime, RollInput, RawRollInput);
+		YawInput = YawInputRate.InterpInputValue(DeltaTime, YawInput, RawYawInput);
+		BoostInput = BoostInputRate.InterpInputValue(DeltaTime, BoostInput, RawBoostInput);
+		DriftInput = DriftInputRate.InterpInputValue(DeltaTime, DriftInput, RawDriftInput);
+		ReverseInput = RawReverseInput;
+
+		if (!bUsingNetworkPhysicsPrediction)
+		{
+			// and send to server - (ServerUpdateState_Implementation below)
+			ServerUpdateState(SteeringInput, ThrottleInput, BrakeInput, HandbrakeInput, -1, RollInput, PitchInput, YawInput, BoostInput, DriftInput, ReverseInput);
+		}
+
+		if (PawnOwner && PawnOwner->IsNetMode(NM_Client))
+		{
+			MarkForClientCameraUpdate();
+		}
+	}
+	else if (!bUsingNetworkPhysicsPrediction)
+	{
+		// use replicated values for remote pawns
+		ThrottleInput = ReplicatedState.Throttle;
+		SteeringInput = ReplicatedState.Steering;
+		BrakeInput = ReplicatedState.Brake;
+		HandbrakeInput = ReplicatedState.Handbrake;
+		PitchInput = ReplicatedState.Pitch;
+		RollInput = ReplicatedState.Roll;
+		YawInput = ReplicatedState.Yaw;
+		BoostInput = ReplicatedState.Boost;
+		DriftInput = ReplicatedState.Drift;
+		ReverseInput = ReplicatedState.Reverse;
+	}
+}
+
+bool UModularVehicleBaseComponent::ServerUpdateState_Validate(float InSteeringInput, float InThrottleInput, float InBrakeInput, float InHandbrakeInput, int32 InCurrentGear, float InRollInput, float InPitchInput, float InYawInput, float InBoostInput, float InDriftInput, bool InReverseInput)
+{
+	return true;
+}
+
+void UModularVehicleBaseComponent::ServerUpdateState_Implementation(float InSteeringInput, float InThrottleInput, float InBrakeInput
+	, float InHandbrakeInput, int32 InCurrentGear, float InRollInput, float InPitchInput, float InYawInput, float InBoostInput, float InDriftInput, bool InReverseInput)
+{
+	SteeringInput = InSteeringInput;
+	ThrottleInput = InThrottleInput;
+	BrakeInput = InBrakeInput;
+	HandbrakeInput = InHandbrakeInput;
+	RollInput = InRollInput;
+	PitchInput = InPitchInput;
+	YawInput = InYawInput;
+	BoostInput = InBoostInput;
+	DriftInput = InDriftInput;
+	ReverseInput = InReverseInput;
+
+	// update state of inputs
+	ReplicatedState.Steering = InSteeringInput;
+	ReplicatedState.Throttle = InThrottleInput;
+	ReplicatedState.Brake = InBrakeInput;
+	ReplicatedState.Handbrake = InHandbrakeInput;
+	ReplicatedState.Roll = InRollInput;
+	ReplicatedState.Pitch = InPitchInput;
+	ReplicatedState.Yaw = InYawInput;
+	ReplicatedState.Boost = InBoostInput;
+	ReplicatedState.Drift = InDriftInput;
+	ReplicatedState.Reverse = InReverseInput;
+}
+
+
+TUniquePtr<FModularVehicleAsyncInput> UModularVehicleBaseComponent::SetCurrentAsyncData(int32 InputIdx, FChaosSimModuleManagerAsyncOutput* CurOutput, FChaosSimModuleManagerAsyncOutput* NextOutput, float Alpha, int32 VehicleManagerTimestamp)
+{
+	TUniquePtr<FModularVehicleAsyncInput> CurInput = MakeUnique<FModularVehicleAsyncInput>();
+	SetCurrentAsyncDataInternal(CurInput.Get(), InputIdx, CurOutput, NextOutput, Alpha, VehicleManagerTimestamp);
+	return CurInput;
+}
+
+/************************************************************************/
+/* Setup the current async I/O data                                     */
+/************************************************************************/
+void UModularVehicleBaseComponent::SetCurrentAsyncDataInternal(FModularVehicleAsyncInput* CurInput, int32 InputIdx, FChaosSimModuleManagerAsyncOutput* CurOutput, FChaosSimModuleManagerAsyncOutput* NextOutput, float Alpha, int32 VehicleManagerTimestamp)
+{
+	ensure(CurAsyncInput == nullptr);	//should be reset after it was filled
+	ensure(CurAsyncOutput == nullptr);	//should get reset after update is done
+
+	CurAsyncInput = CurInput;
+	CurAsyncInput->SetVehicle(this);
+	NextAsyncOutput = nullptr;
+	OutputInterpAlpha = 0.f;
+
+	// We need to find our vehicle in the output given
+	if (CurOutput)
+	{
+		for (int32 PendingOutputIdx = 0; PendingOutputIdx < OutputsWaitingOn.Num(); ++PendingOutputIdx)
+		{
+			// Found the correct pending output, use index to get the vehicle.
+			if (OutputsWaitingOn[PendingOutputIdx].Timestamp == CurOutput->Timestamp)
+			{
+				const int32 VehicleIdx = OutputsWaitingOn[PendingOutputIdx].Idx;
+				FModularVehicleAsyncOutput* VehicleOutput = CurOutput->VehicleOutputs[VehicleIdx].Get();
+				if (VehicleOutput && VehicleOutput->bValid && VehicleOutput->Type == CurAsyncType)
+				{
+					CurAsyncOutput = VehicleOutput;
+
+					if (NextOutput && NextOutput->Timestamp == CurOutput->Timestamp)
+					{
+						// This can occur when sub-stepping - in this case, VehicleOutputs will be in the same order in NextOutput and CurOutput.
+						FModularVehicleAsyncOutput* VehicleNextOutput = NextOutput->VehicleOutputs[VehicleIdx].Get();
+						if (VehicleNextOutput && VehicleNextOutput->bValid && VehicleNextOutput->Type == CurAsyncType)
+						{
+							NextAsyncOutput = VehicleNextOutput;
+							OutputInterpAlpha = Alpha;
+						}
+					}
+				}
+
+				// these are sorted by timestamp, we are using latest, so remove entries that came before it.
+				TArray<FAsyncOutputWrapper> NewOutputsWaitingOn;
+				for (int32 CopyIndex = PendingOutputIdx; CopyIndex < OutputsWaitingOn.Num(); ++CopyIndex)
+				{
+					NewOutputsWaitingOn.Add(OutputsWaitingOn[CopyIndex]);
+				}
+
+				OutputsWaitingOn = MoveTemp(NewOutputsWaitingOn);
+				break;
+			}
+		}
+
+	}
+
+	if (NextOutput && CurOutput)
+	{
+		if (NextOutput->Timestamp != CurOutput->Timestamp)
+		{
+			// NextOutput and CurOutput occurred in different steps, so we need to search for our specific vehicle.
+			for (int32 PendingOutputIdx = 0; PendingOutputIdx < OutputsWaitingOn.Num(); ++PendingOutputIdx)
+			{
+				// Found the correct pending output, use index to get the vehicle.
+				if (OutputsWaitingOn[PendingOutputIdx].Timestamp == NextOutput->Timestamp)
+				{
+					FModularVehicleAsyncOutput* VehicleOutput = NextOutput->VehicleOutputs[OutputsWaitingOn[PendingOutputIdx].Idx].Get();
+					if (VehicleOutput && VehicleOutput->bValid && VehicleOutput->Type == CurAsyncType)
+					{
+						NextAsyncOutput = VehicleOutput;
+						OutputInterpAlpha = Alpha;
+					}
+					break;
+				}
+			}
+		}
+	}
+
+	FAsyncOutputWrapper& NewOutput = OutputsWaitingOn.AddDefaulted_GetRef();
+	NewOutput.Timestamp = VehicleManagerTimestamp;
+	NewOutput.Idx = InputIdx;
+}
+
+inline IPhysicsProxyBase* UModularVehicleBaseComponent::GetPhysicsProxy() const { return ClusterUnionComponent ? ClusterUnionComponent->GetPhysicsProxy() : nullptr; }
+
+
+void UModularVehicleBaseComponent::SetClusterComponent(UClusterUnionComponent* InPhysicalComponent)
+{
+	ClusterUnionComponent = InPhysicalComponent;
+}
+
+/***************************************************************************/
+/* READ OUTPUT DATA - Access the async output data from the Physics Thread */
+/***************************************************************************/
+void UModularVehicleBaseComponent::ParallelUpdate(float DeltaTime)
+{
+	if (FModularVehicleAsyncOutput* CurrentOutput = static_cast<FModularVehicleAsyncOutput*>(CurAsyncOutput))
+	{
+		if (CurrentOutput->bValid && PVehicleOutput)
+		{
+			if (const FModularVehicleAsyncOutput* NextOutput = static_cast<FModularVehicleAsyncOutput*>(NextAsyncOutput))
+			{
+				PVehicleOutput->Clean();
+
+				int NumItems = CurrentOutput->VehicleSimOutput.SimTreeOutputData.Num();
+				PVehicleOutput->SimTreeOutputData.Reserve(NumItems);
+
+				for (int I = 0; I < NumItems; I++)
+				{
+					// #TODO: check we are lerping current/next of the same thing - the number can grow/shrink when construction/destruction happens
+					if (I < NextOutput->VehicleSimOutput.SimTreeOutputData.Num())
+					{
+						Chaos::FSimOutputData* CurrentSimData = CurrentOutput->VehicleSimOutput.SimTreeOutputData[I];
+						Chaos::FSimOutputData* NextSimData = NextOutput->VehicleSimOutput.SimTreeOutputData[I];
+						PVehicleOutput->SimTreeOutputData.EmplaceAt(I, CurrentSimData->MakeNewData());
+						PVehicleOutput->SimTreeOutputData[I]->Lerp(*CurrentSimData, *NextSimData, OutputInterpAlpha);
+
+						// extract/cache some generally useful values as we go as trying to locate this data later requires a search
+						if (PVehicleOutput->SimTreeOutputData[I]->GetType() == Chaos::eSimType::Transmission)
+						{
+							// if there is more than one transmission then the last one will inform us of the current gear
+							CurrentGear = static_cast<Chaos::FTransmissionOutputData*>(PVehicleOutput->SimTreeOutputData[I])->CurrentGear;
+						}
+						else if (PVehicleOutput->SimTreeOutputData[I]->GetType() == Chaos::eSimType::Engine)
+						{
+							// if there is more than one engine then the last one will inform us of the engine RPM
+							Chaos::FEngineOutputData* Engine = static_cast<Chaos::FEngineOutputData*>(PVehicleOutput->SimTreeOutputData[I]);
+
+							EngineRPM = Engine->RPM;
+							EngineTorque = Engine->Torque;
+						}
+
+					}
+				}
+			}
+
+		}
+	}
+}
+
+/************************************************************************/
+/* PASS ANY INPUTS TO THE PHYSICS THREAD SIMULATION IN HERE             */
+/************************************************************************/
+void UModularVehicleBaseComponent::Update(float DeltaTime)
+{
+	if (CurAsyncInput && ClusterUnionComponent && ClusterUnionComponent->GetPhysicsProxy())
+	{
+
+		CurAsyncInput->Proxy = ClusterUnionComponent->GetPhysicsProxy();
+
+		FModularVehicleAsyncInput* AsyncInput = static_cast<FModularVehicleAsyncInput*>(CurAsyncInput);
+
+		AsyncInput->PhysicsInputs.NetworkInputs.VehicleInputs.Throttle = ThrottleInput;
+		AsyncInput->PhysicsInputs.NetworkInputs.VehicleInputs.Boost = BoostInput;
+		AsyncInput->PhysicsInputs.NetworkInputs.VehicleInputs.Drift = DriftInput;
+		AsyncInput->PhysicsInputs.NetworkInputs.VehicleInputs.Brake = BrakeInput;
+		AsyncInput->PhysicsInputs.NetworkInputs.VehicleInputs.Steering = SteeringInput;
+		AsyncInput->PhysicsInputs.NetworkInputs.VehicleInputs.Handbrake = HandbrakeInput;
+		AsyncInput->PhysicsInputs.NetworkInputs.VehicleInputs.Roll = RollInput;
+		AsyncInput->PhysicsInputs.NetworkInputs.VehicleInputs.Pitch = PitchInput;
+		AsyncInput->PhysicsInputs.NetworkInputs.VehicleInputs.Yaw = YawInput;
+		AsyncInput->PhysicsInputs.NetworkInputs.VehicleInputs.Reverse = ReverseInput;
+		AsyncInput->PhysicsInputs.NetworkInputs.VehicleInputs.KeepAwake = bKeepVehicleAwake;
+
+		FCollisionQueryParams TraceParams(NAME_None, FCollisionQueryParams::GetUnknownStatId(), false, nullptr);
+		TraceParams.bReturnPhysicalMaterial = true;	// we need this to get the surface friction coefficient
+		TraceParams.AddIgnoredActors(ActorsToIgnore);
+		TraceParams.bTraceComplex = bSuspensionTraceComplex;
+		AsyncInput->PhysicsInputs.TraceParams = TraceParams;
+		AsyncInput->PhysicsInputs.TraceCollisionResponse = SuspensionTraceCollisionResponses;
+	}
+
+
+}
+
+void UModularVehicleBaseComponent::FinalizeSimCallbackData(FChaosSimModuleManagerAsyncInput& Input)
+{
+	CurAsyncInput = nullptr;
+	CurAsyncOutput = nullptr;
+}
+
+const FTransform& UModularVehicleBaseComponent::GetComponentTransform() const
+{
+	return MyComponent->GetComponentTransform();
+}
+
+
+void UModularVehicleBaseComponent::ActionTreeUpdates(Chaos::FSimTreeUpdates* NextTreeUpdates)
+{
+	Chaos::FClusterUnionPhysicsProxy* Proxy = ClusterUnionComponent->GetPhysicsProxy();
+	Chaos::FPBDRigidsSolver* Solver = Proxy->GetSolver<Chaos::FPBDRigidsSolver>();
+	Solver->EnqueueCommandImmediate([Proxy, this, NextTreeUpdates = *NextTreeUpdates]() mutable
+		{
+			if (VehicleSimulationPT)
+			{
+				TUniquePtr<Chaos::FSimModuleTree>& SimModuleTree = VehicleSimulationPT->AccessSimComponentTree();
+				SimModuleTree->AppendTreeUpdates(NextTreeUpdates);
+				FModularVehicleBuilder::FixupTreeLinks(SimModuleTree);
+
+#if !(UE_BUILD_SHIPPING || UE_BUILD_TEST)
+				if (bModularVehicle_DumpModuleTreeStructure_Enabled)
+				{
+					UE_LOG(LogTemp, Warning, TEXT("SimTreeModules:") );
+					for (int I = 0; I < SimModuleTree->GetNumNodes(); I++)
+					{
+						if (Chaos::ISimulationModuleBase* Module = SimModuleTree->GetNode(I).SimModule)
+						{
+							FString String;
+							Module->GetDebugString(String);
+							UE_LOG(LogTemp, Warning, TEXT("..%s"), *String);
+						}
+					}
+				}
+#endif // !(UE_BUILD_SHIPPING || UE_BUILD_TEST)
+
+				// Network replication data needs to be updated, this is currently studily slow
+				if (NetworkPhysicsComponent)
+				{
+					TSharedPtr<Chaos::FBaseRewindHistory>& History = NetworkPhysicsComponent->GetStatesHistory();
+					Chaos::TDatasRewindHistory<FNetworkModularVehicleStates>* StatesHistory = static_cast<Chaos::TDatasRewindHistory<FNetworkModularVehicleStates>*>(History.Get());
+					if (StatesHistory)
+					{
+						// #TODO: we are rebuilding from scratch every time there is a single change, there must be a better way!
+						// not sure of it is safe to update the datas at this time?
+						for (int I = 0; I < StatesHistory->GetDatasArray().Num(); I++)
+						{
+							FNetworkModularVehicleStates& State = StatesHistory->GetDatasArray()[I];
+							State.ModuleData.Empty();
+
+							VehicleSimulationPT->AccessSimComponentTree()->GenerateReplicationStructure(State.ModuleData);
+						}
+					}
+				}
+			}
+		});
+}
+
+void UModularVehicleBaseComponent::AddComponentToSimulation(UPrimitiveComponent* InComponent, const TArray<FClusterUnionBoneData>& BonesData, const TArray<FClusterUnionBoneData>& RemovedBoneIDs, bool bIsNew)
+{
+	check(ClusterUnionComponent);
+
+	int32 ComponentAddOrder = INDEX_NONE;
+	if (InComponent && bIsNew)
+	{
+		UE_LOG(LogModularBase, Log, TEXT("AddComponentToSimulation %s"), *InComponent->GetName());
+		
+		if (ClusterUnionComponent->IsAuthority())
+		{
+			// retain the order that the components were constructed
+			FConstructionData ConstructionData;
+			ConstructionData.Component = InComponent;
+			ConstructionData.ConstructionIndex = NextConstructionIndex;
+			ConstructionDatas.Add(ConstructionData);
+			ComponentAddOrder = NextConstructionIndex++;
+		}
+		else
+		{
+			ComponentAddOrder = FindComponentAddOrder(InComponent);
+		}
+
+		int ParentID = INDEX_NONE; // always at root
+		Chaos::FSimTreeUpdates LatestTreeUpdates;
+		CreateAssociatedSimComponents(InComponent, ParentID, NextTransformIndex, LatestTreeUpdates);
+
+		Chaos::FSimTreeUpdates* NextTreeUpdates = nullptr;
+		if (LastComponentAddIndex + 1 == ComponentAddOrder)
+		{
+			// has the next in line just come in..
+			NextTreeUpdates = &LatestTreeUpdates;
+			LastComponentAddIndex = ComponentAddOrder;
+			ActionTreeUpdates(NextTreeUpdates);
+		}
+		else
+		{
+			// add to pending and action later
+			PendingTreeUpdates.Add(InComponent, MoveTemp(LatestTreeUpdates));
+		}
+
+		NextTransformIndex++;
+	}
+}
+
+void UModularVehicleBaseComponent::RemoveComponentFromSimulation(UPrimitiveComponent* InComponent, const TArray<FClusterUnionBoneData>& RemovedBonesData)
+{
+	if (InComponent && VehicleSimulationPT)
+	{
+		UE_LOG(LogModularBase, Log, TEXT("UModularVehicleBaseComponent::RemoveComponentFromSimulation: %s"), *InComponent->GetFullName());
+
+		Chaos::FSimTreeUpdates LatestTreeUpdates;
+
+		TArray<UActorComponent*> Components;
+		InComponent->GetOwner()->GetComponents(UVehicleSimBaseComponent::StaticClass(), Components, true);
+
+		for (UActorComponent* ComponentPart : Components)
+		{
+			if (UVehicleSimBaseComponent* ChangedComponent = Cast<UVehicleSimBaseComponent>(ComponentPart))
+			{
+				if (FVehicleComponentData* ComponentData = ComponentToPhysicsObjects.Find(ChangedComponent))
+				{
+					LatestTreeUpdates.RemoveNode(ComponentData->Guid);
+					ComponentToPhysicsObjects.Remove(ChangedComponent);
+				}
+			}
+		}
+
+		TUniquePtr<Chaos::FSimModuleTree>& SimModuleTree = VehicleSimulationPT->AccessSimComponentTree();
+		for (const Chaos::FPendingModuleDeletions& TreeUpdate : LatestTreeUpdates.GetDeletedModules())
+		{
+			for (int Index = 0; Index < SimModuleTree->GetNumNodes(); Index++)
+			{
+				if (Chaos::ISimulationModuleBase* SimModule = SimModuleTree->GetNode(Index).SimModule)
+				{
+					if (SimModule->GetGuid() == TreeUpdate.Guid)
+					{
+						SimModule->SetAnimationEnabled(false);
+						SimModule->SetStateFlags(Chaos::eSimModuleState::Disabled);
+
+						if (SimModule->GetSimType() == Chaos::eSimType::Suspension)
+						{
+							Chaos::FSuspensionSimModule* SuspensionModule = static_cast<Chaos::FSuspensionSimModule*>(SimModule);
+							EnableConstraint(SuspensionModule->GetConstraintIndex(), false);
+						}
+
+						break;
+					}
+				}
+			}
+		}
+
+		Chaos::FClusterUnionPhysicsProxy* Proxy = ClusterUnionComponent->GetPhysicsProxy();
+		Chaos::FPBDRigidsSolver* Solver = Proxy->GetSolver<Chaos::FPBDRigidsSolver>();
+		Solver->EnqueueCommandImmediate([Proxy, this, LatestTreeUpdates = LatestTreeUpdates]() mutable
+			{
+
+				if (VehicleSimulationPT)
+				{
+					TUniquePtr<Chaos::FSimModuleTree>& SimModuleTree = VehicleSimulationPT->AccessSimComponentTree();
+					SimModuleTree->AppendTreeUpdates(LatestTreeUpdates);
+				}
+			});
+	}
+
+}
+
+void UModularVehicleBaseComponent::CreateVehicleSim()
+{
+	UWorld* World = GetWorld();
+
+	// create the simulation class
+	VehicleSimulationPT = MakeUnique<FModularVehicleSimulationCU>(bUsingNetworkPhysicsPrediction);
+
+	// create physics output container
+	PVehicleOutput = MakeUnique<FPhysicsVehicleOutput>();
+
+	// generate the simulation tree with initial components
+	FModularVehicleBuilder::GenerateSimTree(this);
+
+	// register our vehicle with the modular vehicle Sim Manager
+	if (World)
+	{
+		if (World->IsGameWorld())
+		{
+			FPhysScene* PhysScene = World->GetPhysicsScene();
+
+			if (FChaosSimModuleManager* SimManager = FChaosSimModuleManager::GetManagerFromScene(PhysScene))
+			{
+				SimManager->AddVehicle(this);
+			}
+		}
+	}
+
+}
+
+void UModularVehicleBaseComponent::DestroyVehicleSim()
+{
+	UWorld* World = GetWorld();
+	if (World->IsGameWorld())
+	{
+		FPhysScene* PhysScene = World->GetPhysicsScene();
+		check(PhysScene);
+
+		if (Chaos::FPhysicsSolver* LocalSolver = PhysScene->GetSolver())
+		{
+			if (FChaosSimModuleManager* SimManager = FChaosSimModuleManager::GetManagerFromScene(PhysScene))
+			{
+				SimManager->RemoveVehicle(this);
+			}
+
+			LocalSolver->EnqueueCommandImmediate([this, PhysScene = PhysScene]() mutable
+				{
+					if (PVehicleOutput.IsValid())
+					{
+						PVehicleOutput.Reset(nullptr);
+					}
+
+					if (VehicleSimulationPT)
+					{
+						VehicleSimulationPT->Terminate();
+						VehicleSimulationPT.Reset(nullptr);
+					}
+				});
+		}
+
+	}
+
+}
+
+void UModularVehicleBaseComponent::AddOverlappingComponentsToCluster()
+{
+	if (UWorld* World = GetWorld())
+	{
+		FVector HalfSize = AutoAddOverlappingBoxSize * 0.5f;
+		FBox LocalBox = FBox(-HalfSize, HalfSize);
+		TArray<FOverlapResult> OverlapResults;
+		FCollisionShape CollisionBox;
+		CollisionBox.SetBox((FVector3f)LocalBox.GetExtent());
+
+		const FCollisionQueryParams QueryParams;
+		const FCollisionResponseParams ResponseParams;
+		const ECollisionChannel CollisionChannel = ECollisionChannel::ECC_WorldDynamic;
+		const bool bOverlapHit = World->OverlapMultiByChannel(OverlapResults, GetActorLocation(), FQuat::Identity, CollisionChannel, CollisionBox, QueryParams, ResponseParams);
+
+		TArray<int32> BoneIds;
+		BoneIds.Add(0);
+		// Test each overlapped object for a hit result
+		for (FOverlapResult OverlapResult : OverlapResults)
+		{
+			if (UGeometryCollectionComponent* OverlapGCComponent = Cast<UGeometryCollectionComponent>(OverlapResult.Component.Get()))
+			{
+				ClusterUnionComponent->AddComponentToCluster(OverlapGCComponent, BoneIds);
+			}
+		}
+	}
+}
+
+void UModularVehicleBaseComponent::AddGeometryCollectionsFromOwnedActor()
+{
+	TArray<int32> BoneIds;
+	BoneIds.Add(0);
+
+	if (APawn* Pawn = Cast<APawn>(GetOwner()))
+	{
+		if (UClusterUnionComponent* ClusterUnion = Pawn->GetComponentByClass<UClusterUnionComponent>())
+		{
+			TArray<USceneComponent*> ChildComponents;
+			ClusterUnion->GetChildrenComponents(true, ChildComponents);
+
+			for (USceneComponent* Component : ChildComponents)
+			{
+				if (UGeometryCollectionComponent* GCComponent = Cast<UGeometryCollectionComponent>(Component))
+				{
+					ClusterUnionComponent->AddComponentToCluster(GCComponent, BoneIds);
+				}
+			}
+		}
+	}
+}
+
+void UModularVehicleBaseComponent::SetThrottleInput(float Throttle)
+{
+	RawThrottleInput = FMath::Clamp(Throttle, -1.0f, 1.0f);
+}
+
+void UModularVehicleBaseComponent::SetBoostInput(float Boost)
+{
+	RawBoostInput = FMath::Clamp(Boost, -1.0f, 1.0f);
+}
+
+void UModularVehicleBaseComponent::SetDriftInput(float Drift)
+{
+	RawDriftInput = FMath::Clamp(Drift, -1.0f, 1.0f);
+}
+
+void UModularVehicleBaseComponent::IncreaseThrottleInput(float ThrottleDelta)
+{
+	RawThrottleInput = FMath::Clamp(RawThrottleInput + ThrottleDelta, 0.f, 1.0f);
+}
+
+void UModularVehicleBaseComponent::DecreaseThrottleInput(float ThrottleDelta)
+{
+	RawThrottleInput = FMath::Clamp(RawThrottleInput - ThrottleDelta, 0.f, 1.0f);
+}
+
+void UModularVehicleBaseComponent::SetBrakeInput(float Brake)
+{
+	RawBrakeInput = FMath::Clamp(Brake, -1.0f, 1.0f);
+}
+
+void UModularVehicleBaseComponent::SetSteeringInput(float Steering)
+{
+	RawSteeringInput = FMath::Clamp(Steering, -1.0f, 1.0f);
+}
+
+void UModularVehicleBaseComponent::SetPitchInput(float Pitch)
+{
+	RawPitchInput = FMath::Clamp(Pitch, -1.0f, 1.0f);
+}
+
+void UModularVehicleBaseComponent::SetRollInput(float Roll)
+{
+	RawRollInput = FMath::Clamp(Roll, -1.0f, 1.0f);
+}
+
+void UModularVehicleBaseComponent::SetYawInput(float Yaw)
+{
+	RawYawInput = FMath::Clamp(Yaw, -1.0f, 1.0f);
+}
+
+void UModularVehicleBaseComponent::SetHandbrakeInput(float Handbrake)
+{
+	RawHandbrakeInput = Handbrake;
+}
+
+void UModularVehicleBaseComponent::SetReverseInput(bool Reverse)
+{
+	RawReverseInput = Reverse;
+}
+
+
+void UModularVehicleBaseComponent::SetGearInput(int32 Gear)
+{
+	RawGearInput = Gear;
+}
+
+
+int32 UModularVehicleBaseComponent::GetCurrentGear()
+{
+	return CurrentGear;
+}
+
+bool UModularVehicleBaseComponent::IsReversing()
+{
+	return (GetCurrentGear() < 0);
+}
+
+
+//-=====================================================
+// Networking Replication
+
+void UModularVehicleBaseComponent::GetLifetimeReplicatedProps(TArray<FLifetimeProperty>& OutLifetimeProps) const
+{
+	Super::GetLifetimeReplicatedProps(OutLifetimeProps);
+
+	DOREPLIFETIME(UModularVehicleBaseComponent, ReplicatedState);
+	DOREPLIFETIME(UModularVehicleBaseComponent, ConstructionDatas);
+}
+
+void UModularVehicleBaseComponent::ShowDebugInfo(AHUD* HUD, UCanvas* Canvas, const FDebugDisplayInfo& DisplayInfo, float& YL, float& YPos)
+{
+	UFont* RenderFont = GEngine->GetMediumFont();
+
+	// draw input values
+	Canvas->SetDrawColor(FColor::White);
+	YPos += Canvas->DrawText(RenderFont, FString::Printf(TEXT("Throttle Raw  (%f) %f"), RawThrottleInput, ThrottleInput), 4, YPos);
+	YPos += Canvas->DrawText(RenderFont, FString::Printf(TEXT("Boost Raw     (%f) %f"), RawBoostInput, BoostInput), 4, YPos);
+	YPos += Canvas->DrawText(RenderFont, FString::Printf(TEXT("Drift Raw     (%f) %f"), RawDriftInput, DriftInput), 4, YPos);
+	YPos += Canvas->DrawText(RenderFont, FString::Printf(TEXT("Brake Raw     (%f) %f"), RawBrakeInput, BrakeInput), 4, YPos);
+	YPos += Canvas->DrawText(RenderFont, FString::Printf(TEXT("Steering Raw  (%f) %f"), RawSteeringInput, SteeringInput), 4, YPos);
+	YPos += Canvas->DrawText(RenderFont, FString::Printf(TEXT("Handbrake Raw (%f) %f"), RawHandbrakeInput, HandbrakeInput), 4, YPos);
+	YPos += Canvas->DrawText(RenderFont, FString::Printf(TEXT("Roll Raw      (%f) %f"), RawRollInput, RollInput), 4, YPos);
+	YPos += Canvas->DrawText(RenderFont, FString::Printf(TEXT("Pitch Raw     (%f) %f"), RawPitchInput, PitchInput), 4, YPos);
+	YPos += Canvas->DrawText(RenderFont, FString::Printf(TEXT("Yaw Raw       (%f) %f"), RawYawInput, YawInput), 4, YPos);
+	YPos += Canvas->DrawText(RenderFont, FString::Printf(TEXT("Reverse Raw       (%f) %f"), RawReverseInput, ReverseInput), 4, YPos);
+
+	YPos += 10;
+
+}
diff --git a/Engine/Plugins/Experimental/ChaosModularVehicle/Source/ChaosModularVehicleEngine/Private/ChaosModularVehicle/ModularVehicleClusterActor.cpp b/Engine/Plugins/Experimental/ChaosModularVehicle/Source/ChaosModularVehicleEngine/Private/ChaosModularVehicle/ModularVehicleClusterActor.cpp
new file mode 100644
--- /dev/null
+++ b/Engine/Plugins/Experimental/ChaosModularVehicle/Source/ChaosModularVehicleEngine/Private/ChaosModularVehicle/ModularVehicleClusterActor.cpp
@@ -0,0 +1,22 @@
+// Copyright Epic Games, Inc. All Rights Reserved.
+
+#include "ChaosModularVehicle/ModularVehicleClusterActor.h"
+
+#include "PhysicsEngine/ClusterUnionComponent.h"
+#include "ChaosModularVehicle/ModularVehicleBaseComponent.h"
+
+
+AModularVehicleClusterActor::AModularVehicleClusterActor(const FObjectInitializer& ObjectInitializer)
+	: Super(ObjectInitializer)
+{
+	ClusterUnionComponent = CreateDefaultSubobject<UClusterUnionComponent>(TEXT("ClusterUnionComponent0"));
+	SetRootComponent(ClusterUnionComponent);
+
+	// Vehicle sim piggy-backs off of the cluster union component & its events
+	VehicleSimComponent = CreateDefaultSubobject<UModularVehicleBaseComponent>(TEXT("VehicleSimComponent0"));
+	VehicleSimComponent->SetClusterComponent(ClusterUnionComponent);
+
+	SetRemoteRoleForBackwardsCompat(ROLE_SimulatedProxy);
+	bReplicates = true;
+	SetReplicatingMovement(true);
+}
diff --git a/Engine/Plugins/Experimental/ChaosModularVehicle/Source/ChaosModularVehicleEngine/Private/ChaosModularVehicle/ModularVehicleComponent.cpp b/Engine/Plugins/Experimental/ChaosModularVehicle/Source/ChaosModularVehicleEngine/Private/ChaosModularVehicle/ModularVehicleComponent.cpp
new file mode 100644
--- /dev/null
+++ b/Engine/Plugins/Experimental/ChaosModularVehicle/Source/ChaosModularVehicleEngine/Private/ChaosModularVehicle/ModularVehicleComponent.cpp
@@ -0,0 +1,1185 @@
+// Copyright Epic Games, Inc. All Rights Reserved.
+
+#include "ChaosModularVehicle/ModularVehicleComponent.h"
+#include "ChaosModularVehicle/ModularVehicleDefaultAsyncInput.h"
+#include "ChaosModularVehicle/ModularVehicleObject.h"
+#include "Engine/Engine.h"
+#include "SimModule/SimModulesInclude.h"
+#include "ChaosModularVehicle/ChaosSimModuleManager.h"
+#include "ChaosModularVehicle/VehicleSimBaseComponent.h"
+#include "Chaos/DebugDrawQueue.h"
+
+#include "GeometryCollection/GeometryCollectionClusteringUtility.h"
+#include "PhysicsProxy/GeometryCollectionPhysicsProxy.h"
+#include "PhysicsEngine/PhysicsSettings.h"
+#include "PhysicsReplication.h"
+
+#include "VehicleUtility.h"
+#include "Engine/Canvas.h"
+#include "GameFramework/Pawn.h"
+#include "GameFramework/PlayerController.h"
+#include "GameFramework/HUD.h"
+#include "Net/Core/PushModel/PushModel.h"
+#include "Net/UnrealNetwork.h"
+#include "Framework/Threading.h"
+
+#if CHAOS_DEBUG_DRAW
+#include "Chaos/DebugDrawQueue.h"
+#endif
+
+#if VEHICLE_DEBUGGING_ENABLED
+UE_DISABLE_OPTIMIZATION
+#endif
+
+DEFINE_LOG_CATEGORY(LogModularVehicle);
+
+class FModularVehicleSimulation;
+
+bool IsServerNetMode(AActor* Owner)
+{
+	return Owner->IsNetMode(ENetMode::NM_ListenServer) || Owner->IsNetMode(ENetMode::NM_DedicatedServer);
+}
+
+Chaos::FPhysicsSolver* GetSolver(const UModularVehicleComponent& GeometryCollectionComponent)
+{
+	if (GeometryCollectionComponent.ChaosSolverActor)
+	{
+		return GeometryCollectionComponent.ChaosSolverActor->GetSolver();
+	}
+	else if (UWorld* CurrentWorld = GeometryCollectionComponent.GetWorld())
+	{
+		if (FPhysScene* Scene = CurrentWorld->GetPhysicsScene())
+		{
+			return Scene->GetSolver();
+		}
+	}
+	return nullptr;
+}
+
+Chaos::FPBDRigidClusteredParticleHandle* UModularVehicleComponent::GetChassisParticle()
+{
+	Chaos::FPBDRigidClusteredParticleHandle* Particle = nullptr;
+
+	if (FGeometryCollectionPhysicsProxy* Proxy = GetPhysicsProxy())
+	{
+		const TArray<Chaos::FPBDRigidClusteredParticleHandle*>& Clusters = Proxy->GetSolverClusterHandles();
+
+		if (Clusters.Num() > 0 && !Clusters[0]->Disabled())
+		{
+			Particle = Clusters[0];
+		}
+	}
+	return Particle;
+}
+
+UModularVehicleComponent::UModularVehicleComponent(const FObjectInitializer& ObjectInitializer)
+	: Super(ObjectInitializer)
+{
+	PrimaryComponentTick.bCanEverTick = true;
+	bTickInEditor = false;
+
+	bUsingNetworkPhysicsPrediction = UPhysicsSettings::Get()->PhysicsPrediction.bEnablePhysicsPrediction;
+
+	if (bUsingNetworkPhysicsPrediction)
+	{
+		bEnableReplication = true;	// enable GC replication
+	}
+
+	SetIsReplicatedByDefault(true);
+
+	// TODO: currently ordering of this must match EModularVehicleInputType
+	InputInterpolationRates.Reset();
+	InputInterpolationRates.Add(FModularVehicleInputRate(FString("Throttle")));
+	InputInterpolationRates.Add(FModularVehicleInputRate(FString("Brake")));
+	InputInterpolationRates.Add(FModularVehicleInputRate(FString("Clutch")));
+	InputInterpolationRates.Add(FModularVehicleInputRate(FString("Steering")));
+	InputInterpolationRates.Add(FModularVehicleInputRate(FString("Handbrake")));
+	InputInterpolationRates.Add(FModularVehicleInputRate(FString("Pitch")));
+	InputInterpolationRates.Add(FModularVehicleInputRate(FString("Roll")));
+	InputInterpolationRates.Add(FModularVehicleInputRate(FString("Yaw")));
+	InputInterpolationRates.Add(FModularVehicleInputRate(FString("Gear")));
+	InputInterpolationRates.Add(FModularVehicleInputRate(FString("DebugIndex")));
+}
+
+UModularVehicleComponent::~UModularVehicleComponent()
+{
+	
+}
+
+//-=====================================================
+// Networking Replication
+
+void UModularVehicleComponent::GetLifetimeReplicatedProps(TArray<FLifetimeProperty>& OutLifetimeProps) const
+{
+	Super::GetLifetimeReplicatedProps(OutLifetimeProps);
+	DOREPLIFETIME(UModularVehicleComponent, ModularVehicleRepData);
+}
+
+inline void UModularVehicleComponent::UpdateMVRepData()
+{
+	if (APawn* Owner = Cast<APawn>(GetOwner()))
+	{ 
+		APlayerController* OwningPC = Owner ? Cast<APlayerController>(Owner->GetController()) : nullptr;
+		if (OwningPC == nullptr)
+		{
+			return;
+		}
+
+		// If we have no owner or our netmode means we never require replication then early out
+		if (bEnableReplication && IsServerNetMode(Owner) && GetIsReplicated())
+		{
+			MARK_PROPERTY_DIRTY_FROM_NAME(UModularVehicleComponent, ModularVehicleRepData, this);
+			++ModularVehicleRepData.Version;
+
+			ModularVehicleRepData.ServerFrame = Owner->GetWorld()->GetPhysicsScene()->ReplicationCache.ServerFrame;
+
+			//  Controls that were received on the server are now distributed to the clients
+			Chaos::FControlInputs& ControlInputs = VehicleSimulationPT->AccessControlInputs();
+			ModularVehicleRepData.Data.Controls[EModularVehicleInputType::Steering] = ControlInputs.Steering;
+			ModularVehicleRepData.Data.Controls[EModularVehicleInputType::Throttle] = ControlInputs.Throttle;
+			ModularVehicleRepData.Data.Controls[EModularVehicleInputType::Brake] = ControlInputs.Brake;
+			ModularVehicleRepData.Data.Controls[EModularVehicleInputType::Handbrake] = ControlInputs.Handbrake;
+			ModularVehicleRepData.Data.Controls[EModularVehicleInputType::Roll] = ControlInputs.Roll;
+			ModularVehicleRepData.Data.Controls[EModularVehicleInputType::Pitch] = ControlInputs.Pitch;
+			ModularVehicleRepData.Data.Controls[EModularVehicleInputType::Yaw] = ControlInputs.Yaw;
+			ModularVehicleRepData.Data.Controls[EModularVehicleInputType::Gear] = ControlInputs.GearNumber;
+			ModularVehicleRepData.Data.Controls[EModularVehicleInputType::DebugIndex] = ControlInputs.InputDebugIndex;
+		}
+	}
+}
+
+inline void UModularVehicleComponent::ProcessMVRepData()
+{
+	if (UWorld* World = GetWorld())
+	{
+		if (MVRepVersionProcessed < ModularVehicleRepData.Version && GetWorld()->IsNetMode(NM_Client))
+		{
+			MVRepVersionProcessed = ModularVehicleRepData.Version;
+
+			if (APawn* MyPawn = Cast<APawn>(GetOwner()))
+			{
+				APlayerController* OwningPC = MyPawn ? Cast<APlayerController>(MyPawn->GetController()) : nullptr;
+
+				// only update simulated proxy with this particluar data, don't override the local autonimous proxy
+				bool UpdateClient = (OwningPC == nullptr);
+				if (UpdateClient)
+ 				{
+					//  Controls from other clients received for simulates proxies
+					Chaos::FControlInputs& ControlInputs = VehicleSimulationPT->AccessControlInputs();
+					ControlInputs.Steering = ModularVehicleRepData.Data.Controls[EModularVehicleInputType::Steering];
+					ControlInputs.Throttle = ModularVehicleRepData.Data.Controls[EModularVehicleInputType::Throttle];
+					ControlInputs.Brake = ModularVehicleRepData.Data.Controls[EModularVehicleInputType::Brake];
+					ControlInputs.Handbrake = ModularVehicleRepData.Data.Controls[EModularVehicleInputType::Handbrake];
+					ControlInputs.Roll = ModularVehicleRepData.Data.Controls[EModularVehicleInputType::Roll];
+					ControlInputs.Pitch = ModularVehicleRepData.Data.Controls[EModularVehicleInputType::Pitch];
+					ControlInputs.Yaw = ModularVehicleRepData.Data.Controls[EModularVehicleInputType::Yaw];
+					ControlInputs.GearNumber = ModularVehicleRepData.Data.Controls[EModularVehicleInputType::Gear];
+				}
+			}
+		}
+	}
+}
+
+bool UModularVehicleComponent::ProcessRepData(const float DeltaTime, const float SimTime)
+{
+	bool bProcessed = false;
+	const int ChassisIdx = 0;	// TODO: make this more formal
+
+	// With Network Prediction the replication data is largely ignored until the error is greater than the specified tolerance
+	// at which point the client transform will need to be updated to match the authoritative server transform with a resim, but blending away the error over time 
+	// This benefits local control with no input/latency lag.
+	Chaos::FPhysicsSolver* Solver = GetSolver(*this);
+	Chaos::FRewindData* RewindData = Solver->GetRewindData();
+
+	if (bUsingNetworkPhysicsPrediction)
+	{
+		if (RewindData && !RewindData->IsResim())
+		{
+			APawn* MyPawn = Cast<APawn>(GetOwner());
+			APlayerController* OwningPC = MyPawn ? Cast<APlayerController>(MyPawn->GetController()) : nullptr;
+			// no player controller so must be a SimulatedProxy, or if one off event we also need to replicate locally
+			if (OwningPC)
+			{
+				if (UWorld* World = GetWorld())
+				{
+					if (!RepData.Clusters.IsEmpty())
+					{
+						check(World->GetNetMode() == NM_Client);
+
+						if (APlayerController* PlayerController = World->GetFirstPlayerController())
+						{
+							// Same logic as PhysicsReplication.cpp to get the correct local frame in the history buffer
+							int LocalFrameOffset = PlayerController->GetNetworkPhysicsTickOffset();
+							const int32 ServerFrame = RepData.ServerFrame;
+							const int32 LocalFrame = RepData.ServerFrame - LocalFrameOffset;
+						
+							if (FGeometryCollectionPhysicsProxy* Proxy = GetPhysicsProxy())
+							{
+								const FGeometryCollectionClusterRep& TargetState = RepData.Clusters[ChassisIdx];
+
+								// have access to the history buffer
+								check(RewindData);
+								check(!RewindData->IsResim());
+
+								const TArray<Chaos::FPBDRigidClusteredParticleHandle*>& Clusters = Proxy->GetSolverClusterHandles();
+								auto* ChassisHandle = Clusters[ChassisIdx];
+								const FGeometryCollectionClusterRep& RepCluster = RepData.Clusters[ChassisIdx];
+
+								// Need the local frame in order to access correct history data
+								if (ChassisHandle && RewindData && LocalFrame >= RewindData->GetEarliestFrame_Internal() && LocalFrame <= RewindData->CurrentFrame())
+								{
+									static constexpr Chaos::FFrameAndPhase::EParticleHistoryPhase RewindPhase = Chaos::FFrameAndPhase::EParticleHistoryPhase::PostPushData;
+									Chaos::EnsureIsInPhysicsThreadContext();
+
+									Chaos::FGeometryParticleState PastState = RewindData->GetPastStateAtFrame(*ChassisHandle, LocalFrame);
+									int ThisFrame = RewindData->CurrentFrame();
+									Chaos::FGeometryParticleState CurrentState = RewindData->GetPastStateAtFrame(*ChassisHandle, ThisFrame);
+
+									const FVector PastPosition = PastState.X();
+									FVector ErrorVector = (PastPosition - TargetState.Position);
+
+									FQuat ErrorQuat = PastState.R() * TargetState.Rotation.Inverse(); // difference between quaternions
+
+									const float ErrorPosition = ErrorVector.Size();
+
+									const float ErrorThreshold = MyPawn->GetResimulationThreshold();
+									const bool ErrorExceeded = (ErrorPosition >= ErrorThreshold) ? true : false;
+									
+#if CHAOS_DEBUG_DRAW
+									if (Chaos::FPhysicsSolverBase::CanDebugNetworkPhysicsPrediction())
+									{
+										// Server historical position - will appear somewhat lagged to the client
+										float BoxSize1 = 150.0f;
+										float BoxThickness = 6.0f;
+										Chaos::FDebugDrawQueue::GetInstance().DrawDebugBox(TargetState.Position, FVector(BoxSize1, BoxSize1, BoxSize1), TargetState.Rotation, FColor::Black, false, 0.03f, 0, BoxThickness);
+									
+										// client's current position
+										float BoxSize2 = 145.0f;
+										Chaos::FDebugDrawQueue::GetInstance().DrawDebugBox(CurrentState.X(), FVector(BoxSize2, BoxSize2, BoxSize2), CurrentState.R(), FColor::Green, false, 0.03f, 0, BoxThickness);
+									}
+#endif
+
+									// want to always fall through so debug GFX works
+									if (LocalFrame < RewindData->CurrentFrame())
+									{											
+										FAsyncPhysicsTimestamp TimeStamp;
+										TimeStamp.LocalFrame = RewindData->CurrentFrame();
+
+										float ColorLerp = (ErrorPosition >= ErrorThreshold) ? 1.0f : 0.0f;
+
+#if !(UE_BUILD_SHIPPING || UE_BUILD_TEST)
+										if (Chaos::FPhysicsSolverBase::CanDebugNetworkPhysicsPrediction())
+										{
+											UE_LOG(LogModularVehicle, Log, TEXT("Particle Position Error = %f | Resim Trigger = %f | Server Frame = %d | Client Frame = %d"), ErrorPosition, ColorLerp, ServerFrame/*PhysicsTarget.ServerFrame*/, LocalFrame);
+											UE_LOG(LogModularVehicle, Log, TEXT("Particle Target Position = %s | Current Position = %s"), *TargetState.Position.ToString(), *PastState.X().ToString());
+											UE_LOG(LogModularVehicle, Log, TEXT("Particle Target Velocity = %s | Current Velocity = %s"), *TargetState.LinearVelocity.ToString(), *PastState.V().ToString());
+											UE_LOG(LogModularVehicle, Log, TEXT("Particle Target Quaternion = %s | Current Quaternion = %s"), *TargetState.Rotation.ToString(), *PastState.R().ToString());
+											UE_LOG(LogModularVehicle, Log, TEXT("Particle Target Omega = %s | Current Omega= %s"), *TargetState.AngularVelocity.ToString(), *PastState.W().ToString());
+
+#if CHAOS_DEBUG_DRAW
+											if (UWorld* OwningWorld = GetWorld())
+											{
+												static constexpr float BoxSize = 5.0f;
+												const FColor DebugColor = FLinearColor::LerpUsingHSV(FLinearColor::Green, FLinearColor::Red, ColorLerp).ToFColor(false);
+
+												bool PersistLines = true;
+												float DebugRenderFadeTime = 5.0f;
+												Chaos::FDebugDrawQueue::GetInstance().DrawDebugBox(TargetState.Position, FVector(BoxSize, BoxSize, BoxSize), TargetState.Rotation, FColor::Orange, PersistLines, DebugRenderFadeTime, 0, 1.0f);
+												Chaos::FDebugDrawQueue::GetInstance().DrawDebugBox(PastState.X(), FVector(6, 6, 6), PastState.R(), DebugColor, PersistLines, DebugRenderFadeTime, 0, 1.0f);
+
+												Chaos::FDebugDrawQueue::GetInstance().DrawDebugDirectionalArrow(PastState.X(), TargetState.Position, 5.0f, FColor::Red, PersistLines, DebugRenderFadeTime, 0, 0.5f);
+											}
+#endif
+										}
+#endif
+
+										const bool bShouldSleep = false; // #TODO - need equivalent of (TargetState.Flags & ERigidBodyFlags::Sleeping) != 0;
+
+										RewindData->SetTargetStateAtFrame(*GetChassisParticle(), LocalFrame, RewindPhase,
+											TargetState.Position, TargetState.Rotation,
+											TargetState.LinearVelocity, TargetState.AngularVelocity, bShouldSleep);
+
+										if (ColorLerp >= 1.0f)
+										{
+											if (Chaos::FPBDRigidsSolver* RigidSolver = Proxy->GetSolver<Chaos::FPBDRigidsSolver>())
+											{
+												RigidSolver->GetEvolution()->GetIslandManager().SetParticleResimFrame(GetChassisParticle(), LocalFrame);
+											}
+
+											int ResimFrame = RewindData->GetResimFrame();
+											ResimFrame = (ResimFrame == INDEX_NONE) ? LocalFrame : FMath::Min(ResimFrame, LocalFrame);
+											RewindData->SetResimFrame(ResimFrame);
+											UE_LOG(LogModularVehicle, Log, TEXT("### ModularVehicleComponent: Set ResimFrame %d"), ResimFrame);
+										}
+
+									}
+
+									VersionProcessed = RepData.Version;
+									bProcessed = true;
+								}
+							}
+						}
+					}
+				}
+			}
+		}
+	}
+	else
+	{
+		// Client is always corrected to the server position and so is lagged in time
+		return Super::ProcessRepData(DeltaTime, SimTime);
+	}
+
+	return true;
+}
+
+inline void UModularVehicleComponent::AsyncPhysicsTickComponent(float DeltaTime, float SimTime)
+{
+	QUICK_SCOPE_CYCLE_COUNTER(ModularVehicleComponent_AsyncPhysicsTick);
+
+	ENetMode NetMode = GetNetMode();
+
+	// Replication from server to client not required in stand-alone mode
+	if (NetMode != ENetMode::NM_Standalone)
+	{
+		// Geometry Collection Replication
+		Super::AsyncPhysicsTickComponent(DeltaTime, SimTime);
+
+		// Modular Vehicle Replication
+		if (NetMode == ENetMode::NM_Client)
+		{
+			ProcessMVRepData();
+		}
+		else
+		{
+			UpdateMVRepData();
+		}
+	}
+}
+
+//-=====================================================
+
+void UModularVehicleComponent::TickComponent(float DeltaTime, enum ELevelTick TickType, FActorComponentTickFunction* ThisTickFunction)
+{
+	Super::TickComponent(DeltaTime, TickType, ThisTickFunction);
+}
+
+bool UModularVehicleComponent::ShouldCreatePhysicsState() const
+{
+	return true;
+}
+
+void UModularVehicleComponent::UpdateChassisUnionCollision()
+{
+	if (const UGeometryCollection* Rest = GetRestCollection())
+	{
+		TArray<UActorComponent*> Components;
+		GetOwner()->GetComponents(UVehicleSimBaseComponent::StaticClass(), Components, true);
+
+		for (int32 ii = 0, ni = Components.Num(); ii < ni; ++ii)
+		{
+			UVehicleSimBaseComponent* Component = Cast<UVehicleSimBaseComponent>(Components[ii]);
+			if (Component)
+			{
+				if (Component->bRemoveFromClusterCollisionModel)
+				{
+					Rest->GetGeometryCollection()->SetFlags(Component->TransformIndex, FGeometryCollection::ENodeFlags::FS_IgnoreCollisionInParentCluster);
+				}
+				else
+				{
+					Rest->GetGeometryCollection()->ClearFlags(Component->TransformIndex, FGeometryCollection::ENodeFlags::FS_IgnoreCollisionInParentCluster);
+				}
+			}
+		}
+	}
+}
+
+void UModularVehicleComponent::OnCreatePhysicsState()
+{
+	// Selectively disable child collision from the parent cluster union collision geometry
+	// which is create when Super::OnCreatePhysicsState is called
+	UpdateChassisUnionCollision();
+
+	// Geometry Collection Processing
+	Super::OnCreatePhysicsState();
+
+	// Had to move this into PostInitializeComponents since the child components are not yet avaiable/loaded at this time
+	// Vehicle Sim Module Processing
+	//CreateVehicleSim();
+
+}
+
+void UModularVehicleComponent::OnDestroyPhysicsState()
+{
+	// Geometry Collection Processing
+	Super::OnDestroyPhysicsState();
+
+	// Vehicle Sim Module Processing
+	DestroyVehicleSim();
+}
+
+void UModularVehicleComponent::CreateVehicleSim()
+{
+	if (FGeometryCollectionPhysicsProxy* Proxy = GetPhysicsProxy())
+	{
+
+		VehicleSimulationPT = MakeUnique<FModularVehicleSimulationGC>(bUsingNetworkPhysicsPrediction);
+
+		PVehicleOutput = MakeUnique<FPhysicsVehicleOutput>();	// create physics output container
+	
+		GenerateSimTree();
+
+		// Register with the Sim Manager
+		if (UWorld* World = GetWorld())
+		{
+			if (World->IsGameWorld())
+			{
+				FPhysScene* PhysScene = World->GetPhysicsScene();
+
+				if (FChaosSimModuleManager* SimManager = FChaosSimModuleManager::GetManagerFromScene(PhysScene))
+				{
+					SimManager->AddVehicle(this);
+				}
+			}
+		}
+
+		// copy input rate setup across to physics thread
+		TArray<FModularVehicleInputRate>& Inputs = VehicleSimulationPT->AccessInputInterpolation();
+		check(InputInterpolationRates.Num() == Inputs.Num());
+		for (int I = 0; I < Inputs.Num(); I++)
+		{
+			Inputs[I] = InputInterpolationRates[I];
+		}
+	}
+
+}
+
+void UModularVehicleComponent::DestroyVehicleSim()
+{
+	UWorld* World = GetWorld();
+	if (World->IsGameWorld())
+	{
+		FPhysScene* PhysScene = World->GetPhysicsScene();
+
+		if (FChaosSimModuleManager* SimManager = FChaosSimModuleManager::GetManagerFromScene(PhysScene))
+		{
+			SimManager->RemoveVehicle(this);
+		}
+	}
+
+	if (PVehicleOutput.IsValid())
+	{
+		PVehicleOutput.Reset(nullptr);
+	}
+
+	if (VehicleSimulationPT.IsValid())
+	{
+		VehicleSimulationPT->Terminate();
+		VehicleSimulationPT.Reset(nullptr);
+	}
+}
+
+TUniquePtr<Chaos::FSimModuleTree> UModularVehicleComponent::GenerateSimTree()
+{
+	TUniquePtr<Chaos::FSimModuleTree> SimModuleTree = MakeUnique<Chaos::FSimModuleTree>();	// create physics output container
+
+	// take the UVehicleSimBaseComponents and back them into a more compact structure for simulation on physics callback thread.
+	if (SimModuleTree)
+	{
+		int TreeIndex = -1;
+
+		if (const UGeometryCollection* Rest = GetRestCollection())
+		{
+			if (const TSharedPtr<FGeometryCollection, ESPMode::ThreadSafe> GC = Rest->GetGeometryCollection())
+			{
+				if (const UModularVehicle* ModularVehicle = Cast<const UModularVehicle>(Rest))
+				{
+					TArray<UActorComponent*> Components;
+					GetOwner()->GetComponents(UVehicleSimBaseComponent::StaticClass(), Components, true);
+
+					// add a single chassis root component
+					Chaos::FChassisSettings Settings;
+					Chaos::ISimulationModuleBase* Chassis = new Chaos::FChassisSimModule(Settings);
+					Chassis->SetTransformIndex(0);
+					Chassis->SetTreeIndex(0);
+					int RootIndex = SimModuleTree->AddRoot(Chassis);
+
+					// add all other sim components - all at root for now, fixup hierarchy next step (since Component->GetAttachChildren is always empty?!)
+					for (int32 ii = 0, ni = Components.Num(); ii < ni; ++ii)
+					{
+						UVehicleSimBaseComponent* Component = Cast<UVehicleSimBaseComponent>(Components[ii]);
+						if (Component)
+						{
+							int TransformIndex = Component->TransformIndex;
+
+							if (TransformIndex >= GC->Parent.Num())
+							{
+								break; // something is amiss - only support a single sim module per component at present
+							}
+
+							// create simulation data
+							Chaos::ISimulationModuleBase* Module = Component->CreateNewCoreModule();
+							Module->SetSimModuleTree(SimModuleTree.Get());
+
+							TreeIndex = SimModuleTree->AddNodeBelow(RootIndex, Module);
+							Component->TreeIndex = TreeIndex;
+							Module->SetTreeIndex(TreeIndex);
+							Module->SetTransformIndex(TransformIndex);
+
+							if (GC->HasAttribute(TEXT("MassToLocal"), FTransformCollection::TransformGroup) && GC->Parent[TransformIndex] != -1)
+							{
+								const TManagedArray<FTransform>& CollectionMassToLocal = GC->GetAttribute<FTransform>(TEXT("MassToLocal"), FTransformCollection::TransformGroup);
+
+								Module->SetIntactTransform(FTransform::Identity);
+
+								Module->SetClusteredTransform(CollectionMassToLocal[GC->Parent[TransformIndex]].Inverse() * FTransform(GC->Transform[TransformIndex]));
+								Module->SetClustered(true);
+
+							//	UE_LOG(LogSimulationModule, Log, TEXT("[%d] %s GetClusteredTransformLoc %s"), TransformIndex, *Module->GetDebugName(), *Module->GetClusteredTransform().GetLocation().ToString());
+							}
+						}
+					}
+
+					// fix up hierearchy - having to do this as a 2 stage process because Component->GetAttachChildren() is always empty which isn't right 
+					for (int32 ii = 0, ni = Components.Num(); ii < ni; ++ii)
+					{
+						UVehicleSimBaseComponent* Component = Cast<UVehicleSimBaseComponent>(Components[ii]);
+						if (Component)
+						{
+							int ComponentIndex = Component->TreeIndex;
+							if (Component->GetAttachParent() != nullptr)
+							{
+								if (UVehicleSimBaseComponent* Parent = Cast<UVehicleSimBaseComponent>(Component->GetAttachParent()))
+								{
+									SimModuleTree->Reparent(ComponentIndex, Parent->TreeIndex);
+								}
+							}
+						}
+					}
+				}
+			}
+		}
+
+		FixupTreeLinks(SimModuleTree);
+
+		EnableAnimationForPhysicsDrivenTransforms(SimModuleTree);
+
+		// Physics thread takes ownsership of the tree from here
+		VehicleSimulationPT->Initialize(SimModuleTree);
+
+	}
+
+	return SimModuleTree;
+}
+
+void UModularVehicleComponent::EnableAnimationForPhysicsDrivenTransforms(TUniquePtr<Chaos::FSimModuleTree>& SimModuleTree)
+{
+	if (SimModuleTree)
+	{
+		if (FGeometryCollectionPhysicsProxy* Proxy = GetPhysicsProxy())
+		{
+			// AddAttribute/FindAttribute
+			FGeometryDynamicCollection& GTCollection = Proxy->GetExternalCollection();
+
+			GTCollection.AddAnimateTransformAttribute();
+			TManagedArray<bool>* AnimationsActive = GTCollection.GetAnimateTransformAttribute();
+
+			// Find the wheels and set their GC transform to animate back to the GC even when nodes are intact/disabled
+			// TODO: rather than searching for wheels maybe we could have an animate flag on all tree nodes and we search for that being enabled
+			const TArray<Chaos::FSimModuleTree::FSimModuleNode>& ModuleArray = SimModuleTree->GetSimulationModuleTree();
+			for (const Chaos::FSimModuleTree::FSimModuleNode& Node : ModuleArray)
+			{
+				if (Node.SimModule->GetSimType() == Chaos::eSimType::Wheel)
+				{
+					Chaos::FWheelSimModule* Wheel = static_cast<Chaos::FWheelSimModule*>(Node.SimModule);
+
+					int TransformIndex = Node.SimModule->GetTransformIndex();
+					if (TransformIndex != Chaos::ISimulationModuleBase::INVALID_IDX)
+					{
+						(*AnimationsActive)[TransformIndex] = true;
+					}
+				}
+			}
+		}
+	}
+}
+
+
+Chaos::FWheelSimModule* UModularVehicleComponent::LocatePhysicallyClosestWheel(TUniquePtr<Chaos::FSimModuleTree>& SimModuleTree, int SuspensionTreeIndex)
+{
+	const Chaos::FSuspensionSimModule* Suspension = static_cast<const Chaos::FSuspensionSimModule*>(SimModuleTree->GetSimModule(SuspensionTreeIndex));
+
+	FVector SusLoc = Suspension->GetClusteredTransform().GetLocation();
+	float ClosestDistance = FLT_MAX;
+	int ClosestIndex = Chaos::ISimulationModuleBase::INVALID_IDX;
+
+	for (int I = 0; I < SimModuleTree->NumActiveNodes(); I++)
+	{
+		if (SimModuleTree->GetSimModule(I)->GetSimType() == Chaos::eSimType::Wheel)
+		{
+			const Chaos::FWheelSimModule* Wheel = static_cast<const Chaos::FWheelSimModule*>(SimModuleTree->GetSimModule(I));
+			if (Wheel)
+			{
+				FVector WheelLoc = Wheel->GetClusteredTransform().GetLocation();
+				FVector Vec = WheelLoc - SusLoc;
+				float DistSqr = Vec.SquaredLength();
+
+				if (DistSqr < ClosestDistance)
+				{
+					ClosestDistance = DistSqr;
+					ClosestIndex = I;
+				}
+			}
+		}
+	}
+
+	if (ClosestIndex != Chaos::ISimulationModuleBase::INVALID_IDX)
+	{
+		return static_cast<Chaos::FWheelSimModule*>(SimModuleTree->AccessSimModule(ClosestIndex));
+
+	}
+
+	return nullptr;
+}
+
+void UModularVehicleComponent::FixupTreeLinks(TUniquePtr<Chaos::FSimModuleTree>& SimModuleTree)
+{
+	if (const UGeometryCollection* Rest = GetRestCollection())
+	{
+		if (const TSharedPtr<FGeometryCollection, ESPMode::ThreadSafe> GC = Rest->GetGeometryCollection())
+		{
+			for (int I = 0; I < SimModuleTree->NumActiveNodes(); I++)
+			{
+				Chaos::ISimulationModuleBase* Module = SimModuleTree->AccessSimModule(I);
+				if (Module->GetSimType() == Chaos::eSimType::Suspension)
+				{
+					if (Chaos::FSuspensionSimModule* Suspension = static_cast<Chaos::FSuspensionSimModule*>(Module))
+					{
+						if (Chaos::FWheelSimModule* Wheel = LocatePhysicallyClosestWheel(SimModuleTree, Suspension->GetTreeIndex()))
+						{
+							int WheelIndex = Wheel->GetTreeIndex();
+							int WheelTransformIndex = Wheel->GetTransformIndex();
+
+							Suspension->AccessSetup().RestOffset = FVector(GC->Transform[WheelTransformIndex].GetLocation());
+							Suspension->SetIntactTransform(SimModuleTree->AccessSimModule(WheelIndex)->GetIntactTransform());
+							Suspension->SetClusteredTransform(SimModuleTree->AccessSimModule(WheelIndex)->GetClusteredTransform());
+
+							Wheel->SetSuspensionSimTreeIndex(Suspension->GetTreeIndex());
+							Suspension->SetWheelSimTreeIndex(Wheel->GetTreeIndex());
+						}
+					}
+				}
+			}
+		}
+	}
+}
+
+//-=====================================================
+
+void UModularVehicleComponent::SetThrottleInput(float Throttle)
+{
+	RawThrottleInput = FMath::Clamp(Throttle, -1.0f, 1.0f);
+}
+
+void UModularVehicleComponent::IncreaseThrottleInput(float ThrottleDelta)
+{
+	RawThrottleInput = FMath::Clamp(RawThrottleInput + ThrottleDelta, 0.f, 1.0f);
+}
+
+void UModularVehicleComponent::DecreaseThrottleInput(float ThrottleDelta)
+{
+	RawThrottleInput = FMath::Clamp(RawThrottleInput - ThrottleDelta, 0.f, 1.0f);
+}
+
+void UModularVehicleComponent::SetBrakeInput(float Brake)
+{
+	RawBrakeInput = FMath::Clamp(Brake, -1.0f, 1.0f);
+}
+
+void UModularVehicleComponent::SetSteeringInput(float Steering)
+{
+	RawSteeringInput = FMath::Clamp(Steering, -1.0f, 1.0f);
+}
+
+void UModularVehicleComponent::SetPitchInput(float Pitch)
+{
+	RawPitchInput = FMath::Clamp(Pitch, -1.0f, 1.0f);
+}
+
+void UModularVehicleComponent::SetRollInput(float Roll)
+{
+	RawRollInput = FMath::Clamp(Roll, -1.0f, 1.0f);
+}
+
+void UModularVehicleComponent::SetYawInput(float Yaw)
+{
+	RawYawInput = FMath::Clamp(Yaw, -1.0f, 1.0f);
+}
+
+void UModularVehicleComponent::SetHandbrakeInput(float Handbrake)
+{
+	RawHandbrakeInput = Handbrake;
+}
+
+void UModularVehicleComponent::SetGearInput(int Gear)
+{
+	RawGearInput = Gear;
+}
+
+//-=====================================================
+
+void UModularVehicleComponent::PreTickGT(float DeltaTime)
+{
+	ProcessControls(DeltaTime);
+}
+
+TUniquePtr<FModularVehicleAsyncInput> UModularVehicleComponent::SetCurrentAsyncData(int32 InputIdx, FChaosSimModuleManagerAsyncOutput* CurOutput, FChaosSimModuleManagerAsyncOutput* NextOutput, float Alpha, int32 VehicleManagerTimestamp)
+{
+	TUniquePtr<FModularVehicleDefaultAsyncInput> CurInput = MakeUnique<FModularVehicleDefaultAsyncInput>();
+	SetCurrentAsyncDataInternal(CurInput.Get(), InputIdx, CurOutput, NextOutput, Alpha, VehicleManagerTimestamp);
+	return CurInput;
+}
+
+/************************************************************************/
+/* Setup the current async I/O data                                     */
+/************************************************************************/
+void UModularVehicleComponent::SetCurrentAsyncDataInternal(FModularVehicleAsyncInput* CurInput, int32 InputIdx, FChaosSimModuleManagerAsyncOutput* CurOutput, FChaosSimModuleManagerAsyncOutput* NextOutput, float Alpha, int32 VehicleManagerTimestamp)
+{
+	ensure(CurAsyncInput == nullptr);	//should be reset after it was filled
+	ensure(CurAsyncOutput == nullptr);	//should get reset after update is done
+	FModularVehicleDefaultAsyncInput* AsyncInput = static_cast<FModularVehicleDefaultAsyncInput*>(CurInput);
+
+	CurAsyncInput = CurInput;
+	AsyncInput->SetVehicle(this);
+	NextAsyncOutput = nullptr;
+	OutputInterpAlpha = 0.f;
+
+	// We need to find our vehicle in the output given
+	if (CurOutput)
+	{
+		for (int32 PendingOutputIdx = 0; PendingOutputIdx < OutputsWaitingOn.Num(); ++PendingOutputIdx)
+		{
+			// Found the correct pending output, use index to get the vehicle.
+			if (OutputsWaitingOn[PendingOutputIdx].Timestamp == CurOutput->Timestamp)
+			{
+				const int32 VehicleIdx = OutputsWaitingOn[PendingOutputIdx].Idx;
+				FModularVehicleAsyncOutput* VehicleOutput = CurOutput->VehicleOutputs[VehicleIdx].Get();
+				if (VehicleOutput && VehicleOutput->bValid && VehicleOutput->Type == CurAsyncType)
+				{
+					CurAsyncOutput = VehicleOutput;
+
+					if (NextOutput && NextOutput->Timestamp == CurOutput->Timestamp)
+					{
+						// This can occur when sub-stepping - in this case, VehicleOutputs will be in the same order in NextOutput and CurOutput.
+						FModularVehicleAsyncOutput* VehicleNextOutput = NextOutput->VehicleOutputs[VehicleIdx].Get();
+						if (VehicleNextOutput && VehicleNextOutput->bValid && VehicleNextOutput->Type == CurAsyncType)
+						{
+							NextAsyncOutput = VehicleNextOutput;
+							OutputInterpAlpha = Alpha;
+						}
+					}
+				}
+
+				// these are sorted by timestamp, we are using latest, so remove entries that came before it.
+				TArray<FAsyncOutputWrapper> NewOutputsWaitingOn;
+				for (int32 CopyIndex = PendingOutputIdx; CopyIndex < OutputsWaitingOn.Num(); ++CopyIndex)
+				{
+					NewOutputsWaitingOn.Add(OutputsWaitingOn[CopyIndex]);
+				}
+
+				OutputsWaitingOn = MoveTemp(NewOutputsWaitingOn);
+				break;
+			}
+		}
+
+	}
+
+	if (NextOutput && CurOutput)
+	{
+		if (NextOutput->Timestamp != CurOutput->Timestamp)
+		{
+			// NextOutput and CurOutput occurred in different steps, so we need to search for our specific vehicle.
+			for (int32 PendingOutputIdx = 0; PendingOutputIdx < OutputsWaitingOn.Num(); ++PendingOutputIdx)
+			{
+				// Found the correct pending output, use index to get the vehicle.
+				if (OutputsWaitingOn[PendingOutputIdx].Timestamp == NextOutput->Timestamp)
+				{
+					FModularVehicleAsyncOutput* VehicleOutput = NextOutput->VehicleOutputs[OutputsWaitingOn[PendingOutputIdx].Idx].Get();
+					if (VehicleOutput && VehicleOutput->bValid && VehicleOutput->Type == CurAsyncType)
+					{
+						NextAsyncOutput = VehicleOutput;
+						OutputInterpAlpha = Alpha;
+					}
+					break;
+				}
+			}
+		}
+	}
+
+	FAsyncOutputWrapper& NewOutput = OutputsWaitingOn.AddDefaulted_GetRef();
+	NewOutput.Timestamp = VehicleManagerTimestamp;
+	NewOutput.Idx = InputIdx;
+}
+
+// ---- ASYNC ----
+
+FModularVehicleDefaultAsyncInput::FModularVehicleDefaultAsyncInput()
+{
+
+}
+
+/************************************************************************/
+/* Async simulation callback on the Physics Thread                      */
+/************************************************************************/
+TUniquePtr<FModularVehicleAsyncOutput> FModularVehicleDefaultAsyncInput::Simulate(UWorld* World, const float DeltaSeconds, const float TotalSeconds, bool& bWakeOut) const
+{
+	TUniquePtr<FModularVehicleAsyncOutput> Output = MakeUnique<FModularVehicleAsyncOutput>();
+
+	//support nullptr because it allows us to go wide on filling the async inputs
+	if (Proxy == nullptr)
+	{
+		return Output;
+	}
+
+	// FILL OUTPUT DATA HERE THAT WILL GET PASSED BACK TO THE GAME THREAD
+	FGeometryCollectionPhysicsProxy* GCProxy = static_cast<FGeometryCollectionPhysicsProxy*>(Proxy);
+	if (GCProxy)
+	{
+		GCVehicle->VehicleSimulationPT->Simulate(World, DeltaSeconds, *this, *Output.Get(), GCProxy);
+	}
+
+	FModularVehicleAsyncOutput& OutputData = *Output.Get();
+	GCVehicle->VehicleSimulationPT->FillOutputState(OutputData);
+
+	Output->bValid = true;
+
+	return MoveTemp(Output);
+}
+
+void FModularVehicleDefaultAsyncInput::ApplyDeferredForces() const
+{
+	if (Proxy)
+	{
+		FGeometryCollectionPhysicsProxy* GCProxy = static_cast<FGeometryCollectionPhysicsProxy*>(Proxy);
+		if (GCProxy)
+		{
+			GCVehicle->VehicleSimulationPT->ApplyDeferredForces(GCProxy);
+		}
+	}
+}
+
+//void FModularVehicleDefaultAsyncInput::ProcessInputs()
+//{
+//}
+
+
+
+/************************************************************************/
+/* PASS ANY INPUTS TO THE PHYSICS THREAD SIMULATION IN HERE              */
+/************************************************************************/
+void UModularVehicleComponent::Update(float DeltaTime)
+{
+	if (CurAsyncInput)
+	{
+		CurAsyncInput->Proxy = GetPhysicsProxy();
+
+		FModularVehicleDefaultAsyncInput* AsyncInput = static_cast<FModularVehicleDefaultAsyncInput*>(CurAsyncInput);
+
+		TArray<AActor*> ActorsToIgnore;
+		ActorsToIgnore.Add(GetOwner()); // ignore self in scene query
+
+		FCollisionQueryParams TraceParams(NAME_None, FCollisionQueryParams::GetUnknownStatId(), false, nullptr);
+		TraceParams.bReturnPhysicalMaterial = true;	// we need this to get the surface friction coefficient
+		TraceParams.AddIgnoredActors(ActorsToIgnore);
+		TraceParams.bTraceComplex = true;
+		AsyncInput->TraceParams = TraceParams;
+	}
+}
+
+// This is called after the pull from physics state
+void UModularVehicleComponent::PostUpdate()
+{
+	return;
+
+// TODO - code for visual error blending
+//	// want this to run after pull from physics state but before the UpdateGlobalMatricesWithExplodedVectors for render update?
+//	const int ChassisIdx = 0;	// TODO: make this more formal
+//
+//	const float ErrorThreshold = UPhysicsSettings::Get()->PhysicsPrediction.ResimulationErrorThreshold;
+//
+//	NetReplicationError *= 0.98f;
+//	if (NextRepError.Size() > SMALL_NUMBER)
+//	{
+//		NetReplicationError = NextRepError;
+//		NextRepError = FVector::ZeroVector;
+//	}
+//
+//	if (bUsingNetworkPhysicsPrediction)
+//	{
+//		if (Chaos::FPhysicsSolver* Solver = GetSolver(*this))
+//		{
+//			if (FGeometryCollectionPhysicsProxy* Proxy = GetPhysicsProxy())
+//			{
+//				Chaos::FRewindData* RewindData = Solver->GetRewindData();
+//
+//				// [BH] just for debug remove
+//				const TArray<Chaos::FPBDRigidClusteredParticleHandle*>& ClustersTmp = Proxy->GetSolverClusterHandles();
+//				auto* ChassisHandleTmp = ClustersTmp[ChassisIdx];
+//				if (ChassisHandleTmp)
+//				{
+//					UE_LOG(LogModularVehicle, Log, TEXT("***** Pos %s"), *ChassisHandleTmp->X().ToString());
+//
+//					if (!RewindData->IsResim())
+//					{
+//						FVector Delta = PrevPosOther - ChassisHandleTmp->X();
+//						if (Delta.SizeSquared() > (ErrorThreshold * ErrorThreshold))
+//						{
+//							// Warp
+//							UE_LOG(LogModularVehicle, Log, TEXT("***** WARPING %s"), *Delta.ToString());
+//
+//							/*NetReplicationError*/NextRepError = Delta; // delay this by one frame
+//
+//						}
+//
+//						PrevPosOther = ChassisHandleTmp->X();
+//					}
+//				}
+//
+//
+////				Chaos::FRewindData* RewindData = Solver->GetRewindData();
+////				if (RewindData && !RewindData->IsResim())
+////				{			
+////					// [BH - TODO setup ResimFinishFrame any good!!]
+////					if (Warping && ResimFinishFrame == RewindData->CurrentFrame())
+////					{
+////						const TArray<Chaos::FPBDRigidClusteredParticleHandle*>& Clusters = Proxy->GetSolverClusterHandles();
+////						auto* ChassisHandle = Clusters[ChassisIdx];
+////						if (ChassisHandle)
+////						{
+////#ifdef ONWAYORANOTHER
+////							FTransform Loc = GameThreadCollection.Transform[ChassisTransformGroupIndex];
+////							FVector Delta = Loc.GetLocation() - PreRewindPos.GetLocation();
+////#else
+////							// is this thread safe accessing ChassisHandle->X()
+////							FVector Loc = ChassisHandle->X();
+////							FVector Delta = Loc - PreRewindPos;
+////
+////							UE_LOG(LogModularVehicle, Log, TEXT("***** RESIM Final LOC %s"), *PreRewindPos.ToString());
+////							UE_LOG(LogModularVehicle, Log, TEXT("***** Delta %s"), *Delta.ToString());
+////#endif
+////							Warping = false;
+////							NetReplicationError = -Delta;
+////						}
+////					}
+////				}		
+//			}
+//		}
+//	}
+}
+
+void UModularVehicleComponent::FinalizeSimCallbackData(FChaosSimModuleManagerAsyncInput& Input)
+{
+	bool bIsPhysicsStateCreated = true;
+	CurAsyncInput = nullptr;
+	CurAsyncOutput = nullptr;
+}
+
+
+/***************************************************************************/
+/* READ OUTPUT DATA - Access the async output data from the Physics Thread */
+/***************************************************************************/
+void UModularVehicleComponent::ParallelUpdate(float DeltaSeconds)
+{
+	if (FModularVehicleAsyncOutput* CurrentOutput = static_cast<FModularVehicleAsyncOutput*>(CurAsyncOutput))
+	{
+		if (CurrentOutput->bValid && PVehicleOutput)
+		{
+			PVehicleOutput->DebugStrings = MoveTemp(CurrentOutput->VehicleSimOutput.DebugStrings);
+
+			// #TODO: lerp output values here
+			// i.e. PVehicleOutput->Value = FMath::Lerp(CurAsyncOutput->VehicleSimOutput.Value, NextAsyncOutput->VehicleSimOutput.Value, OutputInterpAlpha);
+			//for (int I=0; I<CurAsyncOutput->VehicleSimOutput.ModuleTransform.Num(); I++)
+			//{
+			//	PVehicleOutput->ModuleTransform[I] = CurAsyncOutput->VehicleSimOutput.ModuleTransform[I];
+			//}
+		}
+	}
+}
+
+void UModularVehicleComponent::ClearRawInput()
+{
+	RawThrottleInput = 0.0f;
+	RawBrakeInput = 0.0f;
+	RawSteeringInput = 0.0f;
+	RawClutchInput = 0.0f;	// 0 is engaged/locked, 1 is depressed/open
+	RawPitchInput = 0.0f;
+	RawRollInput = 0.0f;
+	RawYawInput = 0.0f;
+	RawHandbrakeInput = false;
+}
+
+
+void UModularVehicleComponent::ResetVehicleState()
+{
+	ClearRawInput();
+
+	OnDestroyPhysicsState();
+	OnCreatePhysicsState();
+}
+
+
+void UModularVehicleComponent::ShowDebugInfo(AHUD* HUD, UCanvas* Canvas, const FDebugDisplayInfo& DisplayInfo, float& YL, float& YPos)
+{
+	UFont* RenderFont = GEngine->GetMediumFont();
+
+	// draw input values
+	Canvas->SetDrawColor(FColor::White);
+	YPos += Canvas->DrawText(RenderFont, FString::Printf(TEXT("Throttle Raw  (%f) %f"), RawThrottleInput, ThrottleInput), 4, YPos);
+	YPos += Canvas->DrawText(RenderFont, FString::Printf(TEXT("Brake Raw     (%f) %f"), RawBrakeInput, BrakeInput), 4, YPos);
+	YPos += Canvas->DrawText(RenderFont, FString::Printf(TEXT("Steering Raw  (%f) %f"), RawSteeringInput, SteeringInput), 4, YPos);
+	YPos += Canvas->DrawText(RenderFont, FString::Printf(TEXT("Handbrake Raw (%f) %f"), RawHandbrakeInput, HandbrakeInput), 4, YPos);
+	YPos += Canvas->DrawText(RenderFont, FString::Printf(TEXT("Roll Raw      (%f) %f"), RawRollInput, RollInput), 4, YPos);
+	YPos += Canvas->DrawText(RenderFont, FString::Printf(TEXT("Pitch Raw     (%f) %f"), RawPitchInput, PitchInput), 4, YPos);
+	YPos += Canvas->DrawText(RenderFont, FString::Printf(TEXT("Yaw Raw       (%f) %f"), RawYawInput, YawInput), 4, YPos);
+
+	YPos += 10;
+
+	// draw general vehicle data
+	{
+		Canvas->SetDrawColor(FColor::White);
+		YPos += 16;
+
+		for (const FString& StringOut : PVehicleOutput->DebugStrings)
+		{
+			YPos += Canvas->DrawText(RenderFont, FString::Printf(TEXT("%s"), *StringOut), 4, YPos);
+		}
+
+		PVehicleOutput->DebugStrings.Empty();
+	}
+
+	YPos += 10;
+
+}
+
+
+// networking
+
+void UModularVehicleComponent::ProcessControls(float DeltaTime)
+{
+	float PhysicsDeltaTime = DeltaTime;
+	if (UPhysicsSettings::Get())
+	{
+		PhysicsDeltaTime = UPhysicsSettings::Get()->AsyncFixedTimeStepSize;
+	}
+
+	APawn* MyPawn = Cast<APawn>(GetOwner());
+	if (MyPawn)
+	{
+		APlayerController* OwningPC = MyPawn ? Cast<APlayerController>(MyPawn->GetController()) : nullptr;
+
+		if (OwningPC != nullptr && GetOuter() != nullptr && !OwningPC->IsNetMode(NM_DedicatedServer))
+		{
+			const FAsyncPhysicsTimestamp ActualTimestamp = OwningPC->GetPhysicsTimestamp();
+
+			// Apply controls locally
+			if ((OwningPC->GetLocalRole() >= ENetRole::ROLE_AutonomousProxy)&&(OwningPC->GetRemoteRole() != ENetRole::ROLE_AutonomousProxy))
+			{
+				//UE_LOG(LogModularVehicle, Log, TEXT("ApplyLocally [NetMode %d,  LocalRole %d,  GetRemoteRole %d] %f %f"), OwningPC->GetNetMode(), OwningPC->GetLocalRole(), OwningPC->GetRemoteRole(), RawThrottleInput, RawSteeringInput);
+				ApplyControls_Imp(RawSteeringInput, RawThrottleInput, RawBrakeInput, HandbrakeInput, RawGearInput, RawRollInput, RawPitchInput, RawYawInput, ActualTimestamp);
+
+				if (OwningPC->IsNetMode(NM_Client))
+				{
+					// and send to server
+					//UE_LOG(LogModularVehicle, Log, TEXT("SendServerControls [NetMode %d,  LocalRole %d,  GetRemoteRole %d] %f %f"), OwningPC->GetNetMode(), OwningPC->GetLocalRole(), OwningPC->GetRemoteRole(), RawThrottleInput, RawSteeringInput);
+					ServerApplyControls(RawSteeringInput, RawThrottleInput, RawBrakeInput, HandbrakeInput, RawGearInput, RawRollInput, RawPitchInput, RawYawInput, ActualTimestamp);
+				}
+			}
+		}
+	}
+}
+
+void UModularVehicleComponent::ServerApplyControls_Implementation(float InSteeringInput, float InThrottleInput, float InBrakeInput
+	, float InHandbrakeInput, int32 InGearChange, float InRollInput, float InPitchInput, float InYawInput, const FAsyncPhysicsTimestamp& AsyncPhysicsTimestamp)
+{
+	//UE_LOG(LogModularVehicle, Log, TEXT("ApplyFromRPCCallToServer [NetMode %d] %f %f"), GetWorld()->GetNetMode(), RawThrottleInput, RawSteeringInput);
+	ApplyControls_Imp(InSteeringInput, InThrottleInput, InBrakeInput, InHandbrakeInput, InGearChange, InRollInput, InPitchInput, InYawInput, AsyncPhysicsTimestamp);
+}
+
+
+void UModularVehicleComponent::ApplyControls_Imp(float InSteeringInput, float InThrottleInput, float InBrakeInput
+	, float InHandbrakeInput, int32 InGearChange, float InRollInput, float InPitchInput, float InYawInput, const FAsyncPhysicsTimestamp& AsyncPhysicsTimestamp)
+{
+	UWorld* World = GetWorld();
+
+	APawn* MyPawn = Cast<APawn>(GetOwner());
+	APlayerController* OwningPC = MyPawn ? Cast<APlayerController>(MyPawn->GetController()) : nullptr;
+
+	FAsyncPhysicsTimestamp AsyncPhysicsTimestampShift = AsyncPhysicsTimestamp;
+
+	float ControlInputWakeTolerance = 0.02f;
+
+	const bool bControlInputPressed = ((InThrottleInput >= /*GVehicleDebugParams.*/ControlInputWakeTolerance)
+		|| (InBrakeInput >= /*GVehicleDebugParams.*/ControlInputWakeTolerance)
+		|| (InRollInput >= /*GVehicleDebugParams.*/ControlInputWakeTolerance)
+		|| (InPitchInput >= /*GVehicleDebugParams.*/ControlInputWakeTolerance)
+		|| (InYawInput >= /*GVehicleDebugParams.*/ControlInputWakeTolerance)
+		|| (FMath::Abs(InSteeringInput - PrevSteeringInput) >= /*GVehicleDebugParams.*/ControlInputWakeTolerance));
+
+	if (bControlInputPressed)
+	{
+		SetSleeping(false);
+	}
+
+	//UE_LOG(LogModularVehicle, Log, TEXT("Apply : Net Mode %s [ T %f  S %f ]"), *NetModeToString(World->GetNetMode()), InThrottleInput, InSteeringInput);
+
+	OwningPC->ExecuteAsyncPhysicsCommand(AsyncPhysicsTimestampShift, this, [this, InSteeringInput, InThrottleInput, InBrakeInput, InHandbrakeInput, InGearChange, InRollInput, InPitchInput, InYawInput, World, AsyncPhysicsTimestampShift]
+		{
+			Chaos::FControlInputs& ControlInputs = VehicleSimulationPT->AccessControlInputs();
+			ControlInputs.Steering = InSteeringInput;
+			ControlInputs.Throttle = InThrottleInput;
+			ControlInputs.Brake = InBrakeInput;
+			ControlInputs.Handbrake = InHandbrakeInput;
+			ControlInputs.Roll = InRollInput;
+			ControlInputs.Pitch = InPitchInput;
+			ControlInputs.Yaw = InYawInput;
+			ControlInputs.GearNumber = InGearChange;
+			ControlInputs.InputDebugIndex = AsyncPhysicsTimestampShift.LocalFrame;
+
+			// [BH - debug only - try to match controls input on same frame on client and server]
+			//{
+			//	int32 LocalFrameOffset = -1;		// LocalFrame = ServerFrame + LocalFrameOffset;
+
+			//	{
+			//		if (World->GetNetMode() == NM_Client)
+			//		{
+			//			if (APlayerController* PlayerController = World->GetFirstPlayerController())
+			//			{
+			//				//if (FPhysicsSolverBase::IsNetworkPhysicsPredictionEnabled())
+			//				{
+			//					LocalFrameOffset = PlayerController->GetNetworkPhysicsTickOffset();
+			//				}
+			//			}
+			//		}
+			//	}
+
+			//	if (LocalFrameOffset >= 0) // client
+			//	{
+			//		UE_LOG(LogModularVehicle, Log, TEXT("Apply Controls IMP CLIENT Timestamp %d, LocalFrameOffset %d"), AsyncPhysicsTimestampShift.LocalFrame, LocalFrameOffset);
+			//	}
+			//	else // server
+			//	{
+			//		UE_LOG(LogModularVehicle, Log, TEXT("Apply Controls IMP SERVER Timestamp %d"), AsyncPhysicsTimestampShift.LocalFrame);
+			//	}
+			//}
+			// TODO: need to figure out how manual gear changing is going to work, need to find correct node to change state of, i.e. multiple gears ?!?
+		});
+}
+
+void UModularVehicleComponent::SetSleeping(bool bEnableSleep)
+{
+	// TODO
+}
+
+#if VEHICLE_DEBUGGING_ENABLED
+UE_ENABLE_OPTIMIZATION
+#endif
diff --git a/Engine/Plugins/Experimental/ChaosModularVehicle/Source/ChaosModularVehicleEngine/Private/ChaosModularVehicle/ModularVehicleInputRate.cpp b/Engine/Plugins/Experimental/ChaosModularVehicle/Source/ChaosModularVehicleEngine/Private/ChaosModularVehicle/ModularVehicleInputRate.cpp
new file mode 100644
--- /dev/null
+++ b/Engine/Plugins/Experimental/ChaosModularVehicle/Source/ChaosModularVehicleEngine/Private/ChaosModularVehicle/ModularVehicleInputRate.cpp
@@ -0,0 +1,53 @@
+// Copyright Epic Games, Inc. All Rights Reserved.
+
+#include "ChaosModularVehicle/ModularVehicleInputRate.h"
+
+inline float FModularVehicleInputRate::InterpInputValue(float DeltaTime, float CurrentValue, float NewValue) const
+{
+	const float DeltaValue = NewValue - CurrentValue;
+
+	// We are "rising" when DeltaValue has the same sign as CurrentValue (i.e. delta causes an absolute magnitude gain)
+	// OR we were at 0 before, and our delta is no longer 0.
+	const bool bRising = ((DeltaValue > 0.0f) == (CurrentValue > 0.0f)) ||
+		((DeltaValue != 0.f) && (CurrentValue == 0.f));
+
+	const float MaxDeltaValue = DeltaTime * (bRising ? RiseRate : FallRate);
+	const float ClampedDeltaValue = FMath::Clamp(DeltaValue, -MaxDeltaValue, MaxDeltaValue);
+	return CurrentValue + ClampedDeltaValue;
+}
+
+inline float FModularVehicleInputRate::CalcControlFunction(float InputValue)
+{
+	// user defined curve
+
+	// else use option from drop down list
+	switch (InputCurveFunction)
+	{
+	case EModularVehicleInputFunctionType::CustomCurve:
+	{
+		if (UserCurve.GetRichCurveConst() && !UserCurve.GetRichCurveConst()->IsEmpty())
+		{
+			float Output = FMath::Clamp(UserCurve.GetRichCurveConst()->Eval(FMath::Abs(InputValue)), 0.0f, 1.0f);
+			return (InputValue < 0.f) ? -Output : Output;
+		}
+		else
+		{
+			return InputValue;
+		}
+	}
+	break;
+	case EModularVehicleInputFunctionType::SquaredFunction:
+	{
+		return (InputValue < 0.f) ? -InputValue * InputValue : InputValue * InputValue;
+	}
+	break;
+
+	case EModularVehicleInputFunctionType::LinearFunction:
+	default:
+	{
+		return InputValue;
+	}
+	break;
+
+	}
+}
diff --git a/Engine/Plugins/Experimental/ChaosModularVehicle/Source/ChaosModularVehicleEngine/Private/ChaosModularVehicle/ModularVehicleObject.cpp b/Engine/Plugins/Experimental/ChaosModularVehicle/Source/ChaosModularVehicleEngine/Private/ChaosModularVehicle/ModularVehicleObject.cpp
new file mode 100644
--- /dev/null
+++ b/Engine/Plugins/Experimental/ChaosModularVehicle/Source/ChaosModularVehicleEngine/Private/ChaosModularVehicle/ModularVehicleObject.cpp
@@ -0,0 +1,14 @@
+// Copyright Epic Games, Inc. All Rights Reserved.
+
+#include "ChaosModularVehicle/ModularVehicleObject.h"
+
+
+UModularVehicle::UModularVehicle(const FObjectInitializer& ObjectInitializer)
+	: Super(ObjectInitializer)
+{
+	// want regular mass, mass as density yields considerably larger values which don't make sense for vehicles
+	bMassAsDensity = false;
+
+	// default mass value
+	Mass = 1500.0f;
+}
diff --git a/Engine/Plugins/Experimental/ChaosModularVehicle/Source/ChaosModularVehicleEngine/Private/ChaosModularVehicle/ModularVehiclePawn.cpp b/Engine/Plugins/Experimental/ChaosModularVehicle/Source/ChaosModularVehicleEngine/Private/ChaosModularVehicle/ModularVehiclePawn.cpp
new file mode 100644
--- /dev/null
+++ b/Engine/Plugins/Experimental/ChaosModularVehicle/Source/ChaosModularVehicleEngine/Private/ChaosModularVehicle/ModularVehiclePawn.cpp
@@ -0,0 +1,67 @@
+// Copyright Epic Games, Inc. All Rights Reserved.
+
+#include "ChaosModularVehicle/ModularVehiclePawn.h"
+
+
+AModularVehiclePawn::AModularVehiclePawn(const FObjectInitializer& ObjectInitializer)
+	: Super(ObjectInitializer)
+{
+	ModularVehicleComponent = CreateDefaultSubobject<UModularVehicleComponent>(TEXT("ModularVehicleComponent0"));
+
+	RootComponent = ModularVehicleComponent;
+
+	PrimaryActorTick.bCanEverTick = true;
+	SetActorTickEnabled(true);
+
+	bReplicates = true;
+}
+
+void AModularVehiclePawn::PostInitializeComponents()
+{
+	Super::PostInitializeComponents();
+
+	ModularVehicleComponent->CreateVehicleSim();
+}
+
+
+void AModularVehiclePawn::BeginPlay()
+{
+	// COPIED THIS FROM STATIC MESH ACTOR
+	// BEGIN..
+	// Since we allow AStaticMeshActor to specify whether it replicates via bStaticMeshReplicateMovement - per placed instance
+	// and we do the normal SetReplicates call in PostInitProperties, before instanced properties are serialized in, 
+	// we need to do this here. 
+	//
+	// This is a short term fix until we find a better play for SetReplicates to be called in AActor.
+
+	if (GetLocalRole() == ROLE_Authority && true/*bStaticMeshReplicateMovement*/)
+	{
+		bReplicates = false;
+		SetRemoteRoleForBackwardsCompat(ROLE_SimulatedProxy);
+		SetReplicatingMovement(true); //??
+
+		SetReplicates(true);
+	}
+
+	// ..END
+	Super::BeginPlay();
+}
+
+
+void AModularVehiclePawn::TickSimulation(float DeltaTime)
+{
+//	Evolution->AdvanceOneTimeStep(DeltaTime);
+}
+
+
+void AModularVehiclePawn::Tick(float DeltaTime)
+{
+	Super::Tick(DeltaTime);
+
+	if (ModularVehicleComponent)
+	{
+		ModularVehicleComponent->SetRenderStateDirty();
+	}
+}
+
+
