Move Modular Vehicle Plugin files out of restricted folders

diff --git a/Engine/Plugins/Experimental/ChaosModularVehicle/Source/ChaosModularVehicleEngine/Private/ChaosModularVehicle/ModularVehicleSimulationCU.cpp b/Engine/Plugins/Experimental/ChaosModularVehicle/Source/ChaosModularVehicleEngine/Private/ChaosModularVehicle/ModularVehicleSimulationCU.cpp
new file mode 100644
--- /dev/null
+++ b/Engine/Plugins/Experimental/ChaosModularVehicle/Source/ChaosModularVehicleEngine/Private/ChaosModularVehicle/ModularVehicleSimulationCU.cpp
@@ -0,0 +1,402 @@
+// Copyright Epic Games, Inc. All Rights Reserved.
+
+#include "ChaosModularVehicle/ModularVehicleSimulationCU.h"
+#include "ChaosModularVehicle/ModularVehicleDefaultAsyncInput.h"
+#include "PhysicsProxy/GeometryCollectionPhysicsProxy.h"
+#include "PhysicsProxy/ClusterUnionPhysicsProxy.h"
+#include "GeometryCollection/GeometryCollectionObject.h"
+#include "PhysicalMaterials/PhysicalMaterial.h"
+#include "PhysicsEngine/PhysicsSettings.h"
+#include "Chaos/ClusterUnionManager.h"
+#include "Chaos/DebugDrawQueue.h"
+#include "Engine/World.h"
+#include "ChaosModularVehicle/ModularVehicleDebug.h"
+
+FModularVehicleDebugParams GModularVehicleDebugParams;
+
+#if CHAOS_DEBUG_DRAW
+FAutoConsoleVariableRef CVarChaosModularVehiclesShowDebug(TEXT("p.ModularVehicle.ShowDebug"), GModularVehicleDebugParams.ShowDebug, TEXT("Enable/Disable Show Modular Vehicle Debug on HUD."));
+FAutoConsoleVariableRef CVarChaosModularVehiclesRaycastsEnabled(TEXT("p.ModularVehicle.SuspensionRaycastsEnabled"), GModularVehicleDebugParams.SuspensionRaycastsEnabled, TEXT("Enable/Disable Suspension Raycasts."));
+FAutoConsoleVariableRef CVarChaosModularVehiclesShowRaycasts(TEXT("p.ModularVehicle.ShowSuspensionRaycasts"), GModularVehicleDebugParams.ShowSuspensionRaycasts, TEXT("Enable/Disable Suspension Raycast Visualisation."));
+FAutoConsoleVariableRef CVarChaosModularVehiclesShowWheelData(TEXT("p.ModularVehicle.ShowWheelData"), GModularVehicleDebugParams.ShowWheelData, TEXT("Enable/Disable Displaying Wheel Simulation Data."));
+FAutoConsoleVariableRef CVarChaosModularVehiclesShowRaycastMaterial(TEXT("p.ModularVehicle.ShowRaycastMaterial"), GModularVehicleDebugParams.ShowRaycastMaterial, TEXT("Enable/Disable Raycast Material Hit Visualisation."));
+FAutoConsoleVariableRef CVarChaosModularVehiclesShowWheelCollisionNormal(TEXT("p.ModularVehicle.ShowWheelCollisionNormal"), GModularVehicleDebugParams.ShowWheelCollisionNormal, TEXT("Enable/Disable Wheel Collision Normal Visualisation."));
+FAutoConsoleVariableRef CVarChaosModularVehiclesFrictionOverride(TEXT("p.ModularVehicle.FrictionOverride"), GModularVehicleDebugParams.FrictionOverride, TEXT("Override the physics material friction value.."));
+FAutoConsoleVariableRef CVarChaosModularVehiclesDisableAnim(TEXT("p.ModularVehicle.DisableAnim"), GModularVehicleDebugParams.DisableAnim, TEXT("Disable animating wheels, etc"));
+#endif
+
+
+void FModularVehicleSimulationCU::Initialize(TUniquePtr<Chaos::FSimModuleTree>& InSimModuleTree)
+{
+	SimModuleTree = MoveTemp(InSimModuleTree);
+	InputInterpolation.Init(FModularVehicleInputRate(), EModularVehicleInputType::Max);
+}
+
+void FModularVehicleSimulationCU::Terminate()
+{
+	SimModuleTree.Reset(nullptr);
+}
+
+void FModularVehicleSimulationCU::Simulate(UWorld* InWorld, float DeltaSeconds, const FModularVehicleAsyncInput& InputData, FModularVehicleAsyncOutput& OutputData, IPhysicsProxyBase* Proxy)
+{
+#if DEBUG_NETWORK_PHYSICS
+
+	if (InWorld->IsNetMode(NM_ListenServer) || InWorld->IsNetMode(NM_DedicatedServer))
+	{
+		UE_LOG(LogTemp, Log, TEXT("SERVER | PT | TickVehicle | Async tick vehicle with inputs at frame %d : Throttle = %f Brake = %f Roll = %f Pitch = %f Yaw = %f Steering = %f Handbrake = %f"),
+			InputData.PhysicsInputs.NetworkInputs.LocalFrame, VehicleInputs.Throttle, VehicleInputs.Brake, VehicleInputs.Roll, VehicleInputs.Pitch,
+			VehicleInputs.Yaw, VehicleInputs.Steering, VehicleInputs.Handbrake);
+
+		UE_LOG(LogTemp, Log, TEXT("ALT-SERVER | PT | TickVehicle | Async tick vehicle with inputs at frame %d : Throttle = %f Brake = %f Roll = %f Pitch = %f Yaw = %f Steering = %f Handbrake = %f"),
+			InputData.PhysicsInputs.NetworkInputs.LocalFrame, InputData.PhysicsInputs.NetworkInputs.VehicleInputs.Throttle, InputData.PhysicsInputs.NetworkInputs.VehicleInputs.Brake, InputData.PhysicsInputs.NetworkInputs.VehicleInputs.Roll, InputData.PhysicsInputs.NetworkInputs.VehicleInputs.Pitch,
+			InputData.PhysicsInputs.NetworkInputs.VehicleInputs.Yaw, InputData.PhysicsInputs.NetworkInputs.VehicleInputs.Steering, InputData.PhysicsInputs.NetworkInputs.VehicleInputs.Handbrake);
+	}
+	else if (InWorld->IsNetMode(NM_Client))
+	{
+		UE_LOG(LogTemp, Log, TEXT("CLIENT | PT | TickVehicle | Async tick vehicle with inputs at frame %d : Throttle = %f Brake = %f Roll = %f Pitch = %f Yaw = %f Steering = %f Handbrake = %f"),
+			InputData.PhysicsInputs.NetworkInputs.LocalFrame, VehicleInputs.Throttle, VehicleInputs.Brake, VehicleInputs.Roll, VehicleInputs.Pitch,
+			VehicleInputs.Yaw, VehicleInputs.Steering, VehicleInputs.Handbrake);
+
+		UE_LOG(LogTemp, Log, TEXT("ALT-CLIENT | PT | TickVehicle | Async tick vehicle with inputs at frame %d : Throttle = %f Brake = %f Roll = %f Pitch = %f Yaw = %f Steering = %f Handbrake = %f"),
+			InputData.PhysicsInputs.NetworkInputs.LocalFrame, InputData.PhysicsInputs.NetworkInputs.VehicleInputs.Throttle, InputData.PhysicsInputs.NetworkInputs.VehicleInputs.Brake, InputData.PhysicsInputs.NetworkInputs.VehicleInputs.Roll, InputData.PhysicsInputs.NetworkInputs.VehicleInputs.Pitch,
+			InputData.PhysicsInputs.NetworkInputs.VehicleInputs.Yaw, InputData.PhysicsInputs.NetworkInputs.VehicleInputs.Steering, InputData.PhysicsInputs.NetworkInputs.VehicleInputs.Handbrake);
+	}
+
+#endif
+
+
+	Chaos::FPBDRigidsSolver* RigidsSolver = static_cast<Chaos::FPBDRigidsSolver*>(Proxy->GetSolver<Chaos::FPhysicsSolver>());
+	int CurrentFrame = -1;
+	if (RigidsSolver != nullptr)
+	{
+		Chaos::FRewindData* RewindData = RigidsSolver->GetRewindData();
+		if (RewindData != nullptr)
+		{
+			CurrentFrame = RewindData->CurrentFrame();
+		}
+	}
+
+	//if (InWorld)
+	//{
+	//	WriteNetReport(InWorld->IsNetMode(NM_Client), FString::Printf(TEXT("Frame %d,  Throttle %3.2f,  Brake %3.2f, Current Frame %d")
+	//		, InputData.PhysicsInputs.NetworkInputs.LocalFrame
+	//		, InputData.PhysicsInputs.NetworkInputs.VehicleInputs.Throttle
+	//		, InputData.PhysicsInputs.NetworkInputs.VehicleInputs.Brake
+	//		, CurrentFrame));
+	//}
+
+	check(Proxy->GetType() == EPhysicsProxyType::ClusterUnionProxy)
+	Simulate_ClusterUnion(InWorld, DeltaSeconds, InputData, OutputData, static_cast<Chaos::FClusterUnionPhysicsProxy*>(Proxy));
+}
+
+void FModularVehicleSimulationCU::Simulate_ClusterUnion(UWorld* InWorld, float DeltaSeconds, const FModularVehicleAsyncInput& InputData, FModularVehicleAsyncOutput& OutputData, Chaos::FClusterUnionPhysicsProxy* Proxy)
+{
+	Chaos::EnsureIsInPhysicsThreadContext();
+
+	if (Proxy && SimModuleTree.IsValid())
+	{
+		int InitialNum = SimModuleTree->GetSimulationModuleTree().Num();
+
+		//if (InWorld)
+		//{
+		//	WriteNetReport(InWorld->IsNetMode(NM_Client), FString::Printf(TEXT("X %s,  R %s,  V %s,  W %s")
+		//		, *Proxy->GetParticle_Internal()->X().ToString()
+		//		, *Proxy->GetParticle_Internal()->R().ToString()
+		//		, *Proxy->GetParticle_Internal()->V().ToString()
+		//		, *Proxy->GetParticle_Internal()->W().ToString()));
+		//}
+
+		//InterpolateInputs(DeltaSeconds, ExternalInputs, InterpolatedInputs);
+
+		SimInputData.ControlInputs.Throttle = InputData.PhysicsInputs.NetworkInputs.VehicleInputs.Throttle;
+		SimInputData.ControlInputs.Steering = InputData.PhysicsInputs.NetworkInputs.VehicleInputs.Steering;
+		SimInputData.ControlInputs.Brake = InputData.PhysicsInputs.NetworkInputs.VehicleInputs.Brake;
+		SimInputData.ControlInputs.Handbrake = InputData.PhysicsInputs.NetworkInputs.VehicleInputs.Handbrake;
+		SimInputData.ControlInputs.Roll = InputData.PhysicsInputs.NetworkInputs.VehicleInputs.Roll;
+		SimInputData.ControlInputs.Pitch = InputData.PhysicsInputs.NetworkInputs.VehicleInputs.Pitch;
+		SimInputData.ControlInputs.Yaw = InputData.PhysicsInputs.NetworkInputs.VehicleInputs.Yaw;
+		SimInputData.ControlInputs.Boost = InputData.PhysicsInputs.NetworkInputs.VehicleInputs.Boost;
+		SimInputData.ControlInputs.Drift = InputData.PhysicsInputs.NetworkInputs.VehicleInputs.Drift;
+		SimInputData.ControlInputs.IsReversing = InputData.PhysicsInputs.NetworkInputs.VehicleInputs.Reverse;
+		SimInputData.bKeepVehicleAwake = InputData.PhysicsInputs.NetworkInputs.VehicleInputs.KeepAwake;
+
+		PerformAdditionalSimWork(InWorld, InputData, Proxy, SimInputData);
+
+		// run the dynamics simulation, engine, suspension, wheels, aerofoils etc.
+		SimModuleTree->Simulate(DeltaSeconds, SimInputData, Proxy);
+	}
+
+}
+
+//void FModularVehicleSimulationCU::InterpolateInputs(float DeltaSeconds, const Chaos::FControlInputs& ExternalInputIn, Chaos::FControlInputs& InterpolatedInputsInOut)
+//{
+//	InterpolatedInputsInOut.Steering = InputInterpolation[EModularVehicleInputType::Steering].InterpInputValue(DeltaSeconds, InterpolatedInputsInOut.Steering, ExternalInputIn.Steering);
+//	InterpolatedInputsInOut.Throttle = InputInterpolation[EModularVehicleInputType::Throttle].InterpInputValue(DeltaSeconds, InterpolatedInputsInOut.Throttle, ExternalInputIn.Throttle);
+//	InterpolatedInputsInOut.Brake = InputInterpolation[EModularVehicleInputType::Brake].InterpInputValue(DeltaSeconds, InterpolatedInputsInOut.Brake, ExternalInputIn.Brake);
+//	InterpolatedInputsInOut.Handbrake = InputInterpolation[EModularVehicleInputType::Handbrake].InterpInputValue(DeltaSeconds, InterpolatedInputsInOut.Handbrake, ExternalInputIn.Handbrake);
+//	InterpolatedInputsInOut.Pitch = InputInterpolation[EModularVehicleInputType::Pitch].InterpInputValue(DeltaSeconds, InterpolatedInputsInOut.Pitch, ExternalInputIn.Pitch);
+//	InterpolatedInputsInOut.Roll = InputInterpolation[EModularVehicleInputType::Roll].InterpInputValue(DeltaSeconds, InterpolatedInputsInOut.Roll, ExternalInputIn.Roll);
+//	InterpolatedInputsInOut.Yaw = InputInterpolation[EModularVehicleInputType::Yaw].InterpInputValue(DeltaSeconds, InterpolatedInputsInOut.Yaw, ExternalInputIn.Yaw);
+//	InterpolatedInputsInOut.GearNumber = ExternalInputIn.GearNumber;
+//	InterpolatedInputsInOut.InputDebugIndex = ExternalInputIn.InputDebugIndex;
+//}
+
+
+void FModularVehicleSimulationCU::PerformAdditionalSimWork(UWorld* InWorld, const FModularVehicleAsyncInput& InputData, Chaos::FClusterUnionPhysicsProxy* Proxy, Chaos::FAllInputs& AllInputs)
+{
+	using namespace Chaos;
+	check(Proxy);
+	Chaos::EnsureIsInPhysicsThreadContext();
+
+	if (SimModuleTree)
+	{
+		FPBDRigidsEvolutionGBF& Evolution = *static_cast<FPBDRigidsSolver*>(Proxy->GetSolver<FPBDRigidsSolver>())->GetEvolution();
+		FClusterUnionManager& ClusterUnionManager = Evolution.GetRigidClustering().GetClusterUnionManager();
+		const FClusterUnionIndex& CUI = Proxy->GetClusterUnionIndex();
+		if (FClusterUnion* ClusterUnion = ClusterUnionManager.FindClusterUnion(CUI))
+		{
+			FPBDRigidClusteredParticleHandle* ClusterHandle = ClusterUnion->InternalCluster;
+			TArray<FPBDRigidParticleHandle*> Particles = ClusterUnion->ChildParticles;
+			FRigidTransform3 ClusterWorldTM = FRigidTransform3(ClusterHandle->X(), ClusterHandle->R());
+
+			//Chaos::FDebugDrawQueue::GetInstance().DrawDebugCoordinateSystem(ClusterHandle->X(), FRotator(ClusterHandle->R()), 200, false, -1.f, 1.f, 3.0f);
+
+			const TArray<Chaos::FSimModuleTree::FSimModuleNode>& ModuleArray = SimModuleTree->GetSimulationModuleTree();
+
+			for (const Chaos::FSimModuleTree::FSimModuleNode& Node : ModuleArray)
+			{
+				if (Node.IsValid() && Node.SimModule && Node.SimModule->IsEnabled())
+				{
+					FRigidTransform3 Frame = FRigidTransform3::Identity;
+					int TransformIndex = Node.SimModule->GetTransformIndex();
+
+					if (TransformIndex < 0 || TransformIndex >= Particles.Num())
+					{
+						continue;
+					}
+
+					FPBDRigidParticleHandle* Child = Node.SimModule->GetParticleFromUniqueIndex(Node.SimModule->GetParticleIndex().Idx, Particles);
+					if (Child == nullptr)
+					{
+						continue;
+					}
+
+					if (FPBDRigidClusteredParticleHandle* ClusterChild = Child->CastToClustered(); ClusterChild && ClusterChild->IsChildToParentLocked())
+					{
+						Frame = ClusterChild->ChildToParent();
+					}
+					else
+					{
+						Frame = ClusterChild->ChildToParent();
+
+						const FRigidTransform3 ChildWorldTM(Child->X(), Child->R());
+						Frame = ChildWorldTM.GetRelativeTransform(ClusterWorldTM);
+					}
+
+					if (ClusterHandle)
+					{
+						AllInputs.VehicleWorldTransform = ClusterWorldTM;
+
+						if (Node.SimModule->IsClustered() && Node.SimModule->IsBehaviourType(Chaos::eSimModuleTypeFlags::Raycast))
+						{
+							Chaos::FSpringTrace OutTrace;
+							Chaos::FSuspensionSimModule* Suspension = static_cast<Chaos::FSuspensionSimModule*>(Node.SimModule);
+
+							// would be cleaner an faster to just store radius in suspension also
+							float WheelRadius = 0;
+							if (Suspension->GetWheelSimTreeIndex() != Chaos::ISimulationModuleBase::INVALID_IDX)
+							{
+								Chaos::FWheelSimModule* Wheel = static_cast<Chaos::FWheelSimModule*>(ModuleArray[Suspension->GetWheelSimTreeIndex()].SimModule);
+								if (Wheel)
+								{
+									WheelRadius = Wheel->Setup().Radius;
+								}
+							}
+
+							Suspension->GetWorldRaycastLocation(ClusterWorldTM, WheelRadius, OutTrace);
+
+							FVector TraceStart = OutTrace.Start;
+							FVector TraceEnd = OutTrace.End;
+
+							const FCollisionQueryParams& TraceParams = InputData.PhysicsInputs.TraceParams;
+							FVector TraceVector(TraceStart - TraceEnd);
+							FVector TraceNormal = TraceVector.GetSafeNormal();
+
+							FHitResult HitResult = FHitResult();
+							ECollisionChannel SpringCollisionChannel = ECollisionChannel::ECC_WorldDynamic;
+							const FCollisionResponseParams& ResponseParams = InputData.PhysicsInputs.TraceCollisionResponse;
+							if (InWorld)
+							{
+								InWorld->LineTraceSingleByChannel(HitResult, TraceStart, TraceEnd, SpringCollisionChannel, TraceParams, ResponseParams);
+							}
+
+							float Offset = Suspension->Setup().MaxLength;
+							if (HitResult.bBlockingHit && GModularVehicleDebugParams.SuspensionRaycastsEnabled)
+							{
+								Offset = HitResult.Distance - WheelRadius;
+
+								if (Suspension->GetWheelSimTreeIndex() != Chaos::ISimulationModuleBase::INVALID_IDX)
+								{
+									const Chaos::FSimModuleTree::FSimModuleNode& WheelNode = ModuleArray[Suspension->GetWheelSimTreeIndex()];
+
+									Chaos::FWheelSimModule* Wheel = static_cast<Chaos::FWheelSimModule*>(WheelNode.SimModule);
+									if (Wheel && HitResult.PhysMaterial.IsValid())
+									{
+										if (GModularVehicleDebugParams.FrictionOverride > 0)
+										{
+											Wheel->SetSurfaceFriction(GModularVehicleDebugParams.FrictionOverride);
+										}
+										else
+										{
+											Wheel->SetSurfaceFriction(HitResult.PhysMaterial->Friction);
+										}
+									}
+								}
+
+#if CHAOS_DEBUG_DRAW
+								if (GModularVehicleDebugParams.ShowSuspensionRaycasts)
+								{
+									Chaos::FDebugDrawQueue::GetInstance().DrawDebugSphere(HitResult.ImpactPoint, 3, 16, FColor::Red, false, -1.f, 0, 10.f);
+								}
+
+								if (Suspension->GetWheelSimTreeIndex() != Chaos::ISimulationModuleBase::INVALID_IDX)
+								{
+									Chaos::FWheelSimModule* Wheel = static_cast<Chaos::FWheelSimModule*>(ModuleArray[Suspension->GetWheelSimTreeIndex()].SimModule);
+									if (Wheel)
+									{
+										if (GModularVehicleDebugParams.ShowWheelData)
+										{
+											FString TextOut = FString::Format(TEXT("{0}"), { Wheel->GetForceIntoSurface() });
+											Chaos::FDebugDrawQueue::GetInstance().DrawDebugString(HitResult.ImpactPoint + FVec3(0, 50, 50), TextOut, nullptr, FColor::White, -1.f, true, 1.0f);
+										}
+									}
+								}
+
+#endif
+							}
+
+#if CHAOS_DEBUG_DRAW
+							if (GModularVehicleDebugParams.ShowSuspensionRaycasts)
+							{
+								FColor DrawColor = FColor::Green;
+								DrawColor = (HitResult.bBlockingHit) ? FColor::Red : FColor::Green;
+								Chaos::FDebugDrawQueue::GetInstance().DrawDebugLine(TraceStart, TraceEnd, DrawColor, false, -1.f, 0, 2.f);
+								Chaos::FDebugDrawQueue::GetInstance().DrawDebugSphere(TraceStart, 3, 16, FColor::White, false, -1.f, 0, 10.f);
+								Chaos::FDebugDrawQueue::GetInstance().DrawDebugSphere(HitResult.ImpactPoint, 1, 16, FColor::Red, false, -1.f, 0, 10.f);
+								FString TextOut = FString::Format(TEXT("{0}"), { HitResult.Time});
+								Chaos::FDebugDrawQueue::GetInstance().DrawDebugString(HitResult.ImpactPoint + FVec3(0, 50, 50), TextOut, nullptr, FColor::White, -1.f, true, 1.0f);
+							}
+
+							if (GModularVehicleDebugParams.ShowRaycastMaterial)
+							{
+								if (HitResult.PhysMaterial.IsValid())
+								{
+									FDebugDrawQueue::GetInstance().DrawDebugString(HitResult.ImpactPoint, HitResult.PhysMaterial->GetName(), nullptr, FColor::White, -1.f, true, 1.0f);
+								}
+							}
+
+							if (GModularVehicleDebugParams.ShowWheelCollisionNormal)
+							{
+								FVector Pt = HitResult.ImpactPoint;
+								FDebugDrawQueue::GetInstance().DrawDebugLine(Pt, Pt + HitResult.Normal * 20.0f, FColor::Yellow, false, 1.0f, 0, 1.0f);
+								FDebugDrawQueue::GetInstance().DrawDebugSphere(Pt, 5.0f, 4, FColor::White, false, 1.0f, 0, 1.0f);
+							}
+
+#endif
+							Suspension->SetSpringLength(Offset, WheelRadius);
+							FVector Up = ClusterWorldTM.GetUnitAxis(EAxis::Z);
+							Suspension->SetTargetPoint(HitResult.ImpactPoint + Up * WheelRadius, HitResult.ImpactNormal, HitResult.bBlockingHit);
+						}
+
+					}
+				}
+
+			}
+
+		}
+
+	}
+}
+
+/**
+ * ApplyDeferredForces should be called after the ParallelUpdatePT to send the calculated forces to the physics thread serially
+ * as this cannot be done in parallel
+ */
+void FModularVehicleSimulationCU::ApplyDeferredForces(FGeometryCollectionPhysicsProxy* Proxy)
+{
+	Chaos::EnsureIsInPhysicsThreadContext();
+
+	if (SimModuleTree && Proxy)
+	{
+		if (UGeometryCollectionComponent* GCComponent = Cast<UGeometryCollectionComponent>(Proxy->GetOwner()))
+		{
+			check(GCComponent->GetOwner());
+
+			if (const UGeometryCollection* Rest = GCComponent->GetRestCollection())
+			{
+				if (Rest->GetGeometryCollection() && Rest->GetGeometryCollection()->HasAttribute(TEXT("MassToLocal"), FTransformCollection::TransformGroup))
+				{
+					const TManagedArray<FTransform>& CollectionMassToLocal = Rest->GetGeometryCollection()->GetAttribute<FTransform>(TEXT("MassToLocal"), FTransformCollection::TransformGroup);
+
+					SimModuleTree->AccessDeferredForces().Apply(
+						Proxy->GetSolverParticleHandles(),
+						Proxy->GetSolverClusterHandles(),
+						Rest->GetGeometryCollection()->Transform,
+						CollectionMassToLocal,
+						Rest->GetGeometryCollection()->Parent);
+				}
+			}
+		}
+	}
+}
+
+void FModularVehicleSimulationCU::ApplyDeferredForces(Chaos::FClusterUnionPhysicsProxy* Proxy)
+{
+	using namespace Chaos;
+
+	Chaos::EnsureIsInPhysicsThreadContext();
+
+	if (SimModuleTree && Proxy)
+	{
+		// This gives us access to the PT parent cluster and child particles
+		FPBDRigidsEvolutionGBF& Evolution = *static_cast<FPBDRigidsSolver*>(Proxy->GetSolver<FPBDRigidsSolver>())->GetEvolution();
+		FClusterUnionManager& ClusterUnionManager = Evolution.GetRigidClustering().GetClusterUnionManager();
+		const FClusterUnionIndex& CUI = Proxy->GetClusterUnionIndex();
+		if (FClusterUnion* ClusterUnion = ClusterUnionManager.FindClusterUnion(CUI))
+		{
+			FPBDRigidClusteredParticleHandle* ClusterHandle = ClusterUnion->InternalCluster;
+			TArray<FPBDRigidParticleHandle*> Particles = ClusterUnion->ChildParticles;
+
+			TArray<FPBDRigidClusteredParticleHandle*> Clusters;
+			Clusters.Add(ClusterHandle);
+
+			SimModuleTree->AccessDeferredForces().Apply(
+				Particles,
+				Clusters);
+		}
+	}
+}
+
+void FModularVehicleSimulationCU::FillOutputState(FModularVehicleAsyncOutput& Output)
+{
+	if (Chaos::FSimModuleTree* SimTree = GetSimComponentTree().Get())
+	{
+		for (int I = 0; I < SimTree->GetNumNodes(); I++)
+		{
+			if (SimTree->GetSimModule(I))
+			{
+				if (Chaos::FSimOutputData* OutData = SimTree->AccessSimModule(I)->GenerateOutputData())
+				{
+					OutData->FillOutputState(SimTree->GetSimModule(I));
+					Output.VehicleSimOutput.SimTreeOutputData.Add(OutData);
+				}
+			}
+		}
+	}
+}
+
+Chaos::FControlInputs& FModularVehicleSimulationCU::AccessControlInputs()
+{
+	Chaos::EnsureIsInPhysicsThreadContext();
+	return SimModuleTree->GetControlInputs();
+}
diff --git a/Engine/Plugins/Experimental/ChaosModularVehicle/Source/ChaosModularVehicleEngine/Private/ChaosModularVehicle/ModularVehicleSimulationGC.cpp b/Engine/Plugins/Experimental/ChaosModularVehicle/Source/ChaosModularVehicleEngine/Private/ChaosModularVehicle/ModularVehicleSimulationGC.cpp
new file mode 100644
--- /dev/null
+++ b/Engine/Plugins/Experimental/ChaosModularVehicle/Source/ChaosModularVehicleEngine/Private/ChaosModularVehicle/ModularVehicleSimulationGC.cpp
@@ -0,0 +1,340 @@
+// Copyright Epic Games, Inc. All Rights Reserved.
+
+#include "ChaosModularVehicle/ModularVehicleSimulationGC.h"
+#include "ChaosModularVehicle/ModularVehicleDefaultAsyncInput.h"
+#include "PhysicsProxy/GeometryCollectionPhysicsProxy.h"
+#include "GeometryCollection/GeometryCollectionObject.h"
+#include "PhysicalMaterials/PhysicalMaterial.h"
+#include "PhysicsEngine/PhysicsSettings.h"
+#include "Chaos/DebugDrawQueue.h"
+#include "Engine/World.h"
+
+
+void FModularVehicleSimulationGC::Initialize(TUniquePtr<Chaos::FSimModuleTree>& InSimModuleTree)
+{
+	SimModuleTree = MoveTemp(InSimModuleTree);
+	InputInterpolation.Init(FModularVehicleInputRate(), EModularVehicleInputType::Max);
+
+	if (bUsingNetworkPhysicsPrediction)
+	{
+		const int32 PhysicsHistoryLength = UPhysicsSettings::Get()->GetPhysicsHistoryCount();
+		VehicleHistoryBuffer.SetNum(NumHistoryFrames);
+	}
+
+}
+
+void FModularVehicleSimulationGC::Terminate()
+{
+	SimModuleTree.Reset(nullptr);
+}
+
+
+void FModularVehicleSimulationGC::SyncHistoryInputs(float DeltaSeconds, FGeometryCollectionPhysicsProxy* Proxy)
+{
+	check(bUsingNetworkPhysicsPrediction);
+	check(SimModuleTree.IsValid());
+
+	Chaos::FControlInputs& ExternalInputs = SimModuleTree->GetControlInputs();
+
+	auto* RigidSolver = Proxy->GetSolver<Chaos::FPBDRigidsSolver>();
+	if (RigidSolver)
+	{
+		const int32 HistoryFrame = RigidSolver->GetCurrentFrame();
+		int32 EndFrame = StartFrame + NumHistoryFrames;
+
+		// if rewind-resim then use the values from the history buffer
+		if (RigidSolver->GetEvolution()->IsResimming())
+		{
+			if (HistoryFrame >= StartFrame && HistoryFrame < EndFrame)
+			{
+				const int32 LocalFrame = HistoryFrame % NumHistoryFrames;
+				if (NumHistoryFrames > LocalFrame)
+				{
+					// interpolation value for this frame is already included in history
+					SimInputData.ControlInputs = VehicleHistoryBuffer[LocalFrame].ControlInputs;
+				}
+
+			}
+		}
+		else
+		{
+			// interpolate control on the PT for determinism
+			// TODO: pass the initial InputRate values over from the GT
+			Chaos::FControlInputs& InterpolatedInputs = SimInputData.ControlInputs;
+			InterpolateInputs(DeltaSeconds, ExternalInputs, InterpolatedInputs);
+
+			if (HistoryFrame >= EndFrame)
+			{
+				StartFrame += HistoryFrame - EndFrame + 1;
+			}
+			EndFrame = StartFrame + NumHistoryFrames;
+
+			if (HistoryFrame >= StartFrame && HistoryFrame < EndFrame)
+			{
+				const int32 LocalFrame = HistoryFrame % NumHistoryFrames;
+				if (NumHistoryFrames > LocalFrame)
+				{
+					// store the interpolated value - we won't interpolate again when using the history
+					VehicleHistoryBuffer[LocalFrame].ControlInputs = SimInputData.ControlInputs;
+				}
+			}
+		}
+	}
+}
+
+/** This function will be called in parallel with other vehicle instances
+  *  - for this reason care must be taken, i.e. can't apply forces directly they must be accumulated/deferred
+  * then applied after the parallel update
+  */
+void FModularVehicleSimulationGC::Simulate(UWorld* InWorld, float DeltaSeconds, const FModularVehicleDefaultAsyncInput& InputData, FModularVehicleAsyncOutput& OutputData, FGeometryCollectionPhysicsProxy* Proxy)
+{
+	Chaos::EnsureIsInPhysicsThreadContext();
+
+	// for debugging to determine if this is the server or client
+	//ENetMode NetMode = InWorld->GetNetMode();
+	//UE_LOG(LogModularVehicle, Log, TEXT("FModularVehicleSimulationGC::Simulate: NetMode %d"), NetMode);
+
+	if (Proxy && SimModuleTree.IsValid())
+	{
+		if (bUsingNetworkPhysicsPrediction)
+		{ 
+			SyncHistoryInputs(DeltaSeconds, Proxy);
+		}
+		else
+		{
+			Chaos::FControlInputs& InterpolatedInputs = SimInputData.ControlInputs;
+			Chaos::FControlInputs& ExternalInputs = SimModuleTree->GetControlInputs();
+			InterpolateInputs(DeltaSeconds, ExternalInputs, InterpolatedInputs);
+		}
+
+		// Processing that currently doesn't quite fit inside the sim tree simulate call
+		PerformAdditionalSimWork(InWorld, InputData, Proxy, SimInputData);
+		
+		int SolverFrame = Proxy->GetSolver<Chaos::FPhysicsSolver>()->GetCurrentFrame();
+		//UE_LOG(LogModularVehicle, Log, TEXT(" ~~~ NetMode %d Frame %d ControlsFrame %d, CheckSum %d, Steering %f")
+		//	, InWorld->GetNetMode()
+		//	, SolverFrame
+		//	, SimInputData.ControlInputs.InputDebugIndex
+		//	, SimInputData.ControlInputs.GetChecksum()
+		//	, SimInputData.ControlInputs.Steering);
+		
+		// run the dynamics simulation, engine, suspension, wheels, aerofoils etc.
+ 		SimModuleTree->Simulate(DeltaSeconds, SimInputData, Proxy);
+	}
+
+}
+
+void FModularVehicleSimulationGC::InterpolateInputs(float DeltaSeconds, const Chaos::FControlInputs& ExternalInputIn, Chaos::FControlInputs& InterpolatedInputsInOut)
+{
+	InterpolatedInputsInOut.Steering = InputInterpolation[EModularVehicleInputType::Steering].InterpInputValue(DeltaSeconds, InterpolatedInputsInOut.Steering, ExternalInputIn.Steering);
+	InterpolatedInputsInOut.Throttle = InputInterpolation[EModularVehicleInputType::Throttle].InterpInputValue(DeltaSeconds, InterpolatedInputsInOut.Throttle, ExternalInputIn.Throttle);
+	InterpolatedInputsInOut.Brake = InputInterpolation[EModularVehicleInputType::Brake].InterpInputValue(DeltaSeconds, InterpolatedInputsInOut.Brake, ExternalInputIn.Brake);
+	InterpolatedInputsInOut.Handbrake = InputInterpolation[EModularVehicleInputType::Handbrake].InterpInputValue(DeltaSeconds, InterpolatedInputsInOut.Handbrake, ExternalInputIn.Handbrake);
+	InterpolatedInputsInOut.Pitch = InputInterpolation[EModularVehicleInputType::Pitch].InterpInputValue(DeltaSeconds, InterpolatedInputsInOut.Pitch, ExternalInputIn.Pitch);
+	InterpolatedInputsInOut.Roll = InputInterpolation[EModularVehicleInputType::Roll].InterpInputValue(DeltaSeconds, InterpolatedInputsInOut.Roll, ExternalInputIn.Roll);
+	InterpolatedInputsInOut.Yaw = InputInterpolation[EModularVehicleInputType::Yaw].InterpInputValue(DeltaSeconds, InterpolatedInputsInOut.Yaw, ExternalInputIn.Yaw);
+	InterpolatedInputsInOut.GearNumber = ExternalInputIn.GearNumber;
+	InterpolatedInputsInOut.InputDebugIndex = ExternalInputIn.InputDebugIndex;
+}
+
+void FModularVehicleSimulationGC::PerformAdditionalSimWork(UWorld* InWorld, const FModularVehicleDefaultAsyncInput& InputData, FGeometryCollectionPhysicsProxy* Proxy, Chaos::FAllInputs& AllInputs)
+{
+	check(Proxy);
+	Chaos::EnsureIsInPhysicsThreadContext();
+
+	FGeometryDynamicCollection& GeometryCollection = Proxy->GetPhysicsCollection();
+	const TArray<Chaos::FPBDRigidClusteredParticleHandle*>& Clusters = Proxy->GetSolverClusterHandles();
+	const TArray<Chaos::FPBDRigidClusteredParticleHandle*>& Particles = Proxy->GetSolverParticleHandles();
+
+	if (Clusters.Num() > 0)
+	{
+		const TArray<Chaos::FSimModuleTree::FSimModuleNode>& ModuleArray = SimModuleTree->GetSimulationModuleTree();
+
+		for (const Chaos::FSimModuleTree::FSimModuleNode& Node : ModuleArray)
+		{
+			if (!Node.SimModule->IsEnabled() || !Particles[Node.SimModule->GetTransformIndex()]->Disabled())
+			{
+				Node.SimModule->SetStateFlags(Chaos::eSimModuleState::Disabled);
+				continue;
+			}
+
+			check(Node.SimModule->GetTransformIndex() < Clusters.Num());
+			const Chaos::FPBDRigidClusteredParticleHandle* Particle = Clusters[Node.SimModule->GetTransformIndex()];
+			if (Particle)
+			{
+				const FTransform BodyTransform(Particle->R(), Particle->X());
+
+				// #TODO: cheating just now to get it working - do we pass in one or one per particle, as it will be different if it has split into fragments
+				AllInputs.VehicleWorldTransform = BodyTransform;
+
+				if (Node.SimModule->IsBehaviourType(Chaos::eSimModuleTypeFlags::Raycast))
+				{
+					Chaos::FSpringTrace OutTrace;
+					Chaos::FSuspensionSimModule* Suspension = static_cast<Chaos::FSuspensionSimModule*>(Node.SimModule);
+
+					// would be cleaner an faster to just store radius in suspension also
+					float WheelRadius = 0;
+					if (Suspension->GetWheelSimTreeIndex() != Chaos::ISimulationModuleBase::INVALID_IDX)
+					{
+						Chaos::FWheelSimModule* Wheel = static_cast<Chaos::FWheelSimModule*>(ModuleArray[Suspension->GetWheelSimTreeIndex()].SimModule);
+						if (Wheel)
+						{
+							WheelRadius = Wheel->Setup().Radius;
+						}
+					}
+
+					Suspension->GetWorldRaycastLocation(BodyTransform, WheelRadius, OutTrace);
+
+					FVector TraceStart = OutTrace.Start;
+					FVector TraceEnd = OutTrace.End;
+
+					const FCollisionQueryParams& TraceParams = InputData.TraceParams;
+					FVector TraceVector(TraceStart - TraceEnd);
+					FVector TraceNormal = TraceVector.GetSafeNormal();
+
+					FHitResult HitResult = FHitResult();
+					ECollisionChannel SpringCollisionChannel = ECollisionChannel::ECC_WorldDynamic;
+					FCollisionResponseParams ResponseParams;
+					ResponseParams.CollisionResponse = FCollisionResponseContainer::GetDefaultResponseContainer();
+					ResponseParams.CollisionResponse.Vehicle = ECR_Ignore;
+					if (InWorld)
+					{
+						InWorld->LineTraceSingleByChannel(HitResult, TraceStart, TraceEnd, SpringCollisionChannel, TraceParams, ResponseParams);
+					}
+
+#if CHAOS_DEBUG_DRAW
+					//if (GModularVehicleDebugParams.ShowRaycasts)
+					//{
+					//	Chaos::FDebugDrawQueue::GetInstance().DrawDebugLine(TraceStart, TraceEnd, FColor::Yellow, false, -1.f, 0, 2.f);
+					//}
+#endif
+
+					float Offset = Suspension->Setup().MaxLength;
+					if (HitResult.bBlockingHit)
+					{
+						FVector LocalPos = Suspension->GetClusteredTransform().GetLocation();
+						FVector LocalHitPoint = BodyTransform.InverseTransformPosition(HitResult.ImpactPoint);
+
+						//Suspension->SetLocation(LocalHitPoint - LocalPos + WheelRadius);
+						Offset = HitResult.Distance - WheelRadius;
+
+						if (Suspension->GetWheelSimTreeIndex() != Chaos::ISimulationModuleBase::INVALID_IDX)
+						{
+							const Chaos::FSimModuleTree::FSimModuleNode& WheelNode = ModuleArray[Suspension->GetWheelSimTreeIndex()];
+
+							Chaos::FWheelSimModule* Wheel = static_cast<Chaos::FWheelSimModule*>(WheelNode.SimModule);
+							if (Wheel)
+							{
+								Wheel->SetSurfaceFriction(HitResult.PhysMaterial->Friction);
+							}
+						}
+
+#if CHAOS_DEBUG_DRAW
+						//if (GModularVehicleDebugParams.ShowRaycasts)
+						//{
+						//	Chaos::FDebugDrawQueue::GetInstance().DrawDebugSphere(HitResult.ImpactPoint, 10, 16, FColor::White, false, -1.f, 0, 10.f);
+						//}
+#endif
+					}
+
+					Suspension->SetSpringLength(Offset, WheelRadius);
+				}
+
+				if (Node.SimModule->GetSimType() == Chaos::eSimType::Wheel)
+				{
+					Chaos::FWheelSimModule* Wheel = static_cast<Chaos::FWheelSimModule*>(Node.SimModule);
+
+					int TransformIndex = Node.SimModule->GetTransformIndex();
+
+					// animate wheel rotation and steering
+					FQuat Rot = FQuat(FVector(0, 1, 0), Wheel->GetAngularPosition());
+					FQuat Steer = FQuat(FVector(0, 0, 1), FMath::DegreesToRadians(Wheel->GetSteerAngleDegrees()));
+					FTransform3f Transform = GeometryCollection.GetTransform(TransformIndex);
+					Transform.SetRotation(FQuat4f(Node.SimModule->GetClusteredTransform().GetRotation() * Steer * Rot));
+					GeometryCollection.SetTransform(TransformIndex, Transform);
+
+					// animate suspension/ wheel
+					if (Wheel->GetSuspensionSimTreeIndex() != Chaos::ISimulationModuleBase::INVALID_IDX)
+					{
+						const Chaos::ISimulationModuleBase* SuspModule = SimModuleTree->GetSimModule(Wheel->GetSuspensionSimTreeIndex());
+						const Chaos::FSuspensionSimModule* Suspension = static_cast<const Chaos::FSuspensionSimModule*>(SuspModule);
+						float CurrentSpringLength = Suspension->GetSpringLength();
+						FVector Loc = FVector(GeometryCollection.GetTransform(Node.SimModule->GetTransformIndex()).GetLocation());
+
+						FVector RestPos = Suspension->GetRestLocation();
+
+						Loc.Z = Suspension->Setup().MaxRaise + RestPos.Z + CurrentSpringLength;
+						Transform.SetLocation(FVector3f(Loc));
+						GeometryCollection.SetTransform(TransformIndex, Transform);
+
+						GeometryCollection.MakeDirty();
+					}
+				}
+			}
+
+		}
+
+	}
+
+}
+
+/**
+ * ApplyDeferredForces should be called after the ParallelUpdatePT to send the calculated forces to the physics thread serially
+ * as this cannot be done in parallel
+ */
+void FModularVehicleSimulationGC::ApplyDeferredForces(FGeometryCollectionPhysicsProxy* Proxy)
+{
+	Chaos::EnsureIsInPhysicsThreadContext();
+
+	if (SimModuleTree && Proxy)
+	{
+		// TODO: shouldn't be accessing Component
+		if (UGeometryCollectionComponent* GCComponent = Cast<UGeometryCollectionComponent>(Proxy->GetOwner()))
+		{
+			check(GCComponent->GetOwner());
+
+			if (const UGeometryCollection* Rest = GCComponent->GetRestCollection())
+			{
+				if (Rest->GetGeometryCollection() && Rest->GetGeometryCollection()->HasAttribute(TEXT("MassToLocal"), FTransformCollection::TransformGroup))
+				{
+					const TManagedArray<FTransform>& CollectionMassToLocal = Rest->GetGeometryCollection()->GetAttribute<FTransform>(TEXT("MassToLocal"), FTransformCollection::TransformGroup);
+
+					SimModuleTree->AccessDeferredForces().Apply(
+						Proxy->GetSolverParticleHandles(),
+						Proxy->GetSolverClusterHandles(),
+						Rest->GetGeometryCollection()->Transform,
+						CollectionMassToLocal,
+						Rest->GetGeometryCollection()->Parent);
+				}
+			}
+		}
+	}
+}
+
+void FModularVehicleSimulationGC::FillOutputState(FModularVehicleAsyncOutput& Output)
+{
+	// #TODO: what output values are we going to have for SimModules
+	// #Note: remember to copy/interpolate values from the physics thread output in UModularVehicleComponent::ParallelUpdate
+
+	//if (GModularVehicleDebugParams.ShowDebug)
+	//{
+	//	Output.VehicleSimOutput.DebugStrings.Empty();
+	//	if (Chaos::FSimModuleTree* SimTree = GetSimComponentTree().Get())
+	//	{
+	//		for (int I = 0; I < SimTree->GetNumNodes(); I++)
+	//		{
+	//			FString NewString;
+	//			if (SimTree->GetSimModule(I)->GetDebugString(NewString))
+	//			{
+	//				Output.VehicleSimOutput.DebugStrings.Add(NewString);
+	//			}
+	//		}
+	//	}
+	//}
+}
+
+Chaos::FControlInputs& FModularVehicleSimulationGC::AccessControlInputs()
+{
+	Chaos::EnsureIsInPhysicsThreadContext();
+	return SimModuleTree->GetControlInputs();
+}
diff --git a/Engine/Plugins/Experimental/ChaosModularVehicle/Source/ChaosModularVehicleEngine/Private/ChaosModularVehicle/VehicleSimBaseComponent.cpp b/Engine/Plugins/Experimental/ChaosModularVehicle/Source/ChaosModularVehicleEngine/Private/ChaosModularVehicle/VehicleSimBaseComponent.cpp
new file mode 100644
--- /dev/null
+++ b/Engine/Plugins/Experimental/ChaosModularVehicle/Source/ChaosModularVehicleEngine/Private/ChaosModularVehicle/VehicleSimBaseComponent.cpp
@@ -0,0 +1,14 @@
+// Copyright Epic Games, Inc. All Rights Reserved.
+
+#include "ChaosModularVehicle/VehicleSimBaseComponent.h"
+
+
+UVehicleSimBaseComponent::UVehicleSimBaseComponent(const FObjectInitializer& ObjectInitializer)
+	: Super(ObjectInitializer)
+{
+	TransformIndex = 0;
+	bRemoveFromClusterCollisionModel = false;
+	TreeIndex = -1;
+	bAnimationEnabled = false;
+}
+
diff --git a/Engine/Plugins/Experimental/ChaosModularVehicle/Source/ChaosModularVehicleEngine/Public/ChaosModularVehicle/ChaosSimModuleManager.h b/Engine/Plugins/Experimental/ChaosModularVehicle/Source/ChaosModularVehicleEngine/Public/ChaosModularVehicle/ChaosSimModuleManager.h
new file mode 100644
--- /dev/null
+++ b/Engine/Plugins/Experimental/ChaosModularVehicle/Source/ChaosModularVehicleEngine/Public/ChaosModularVehicle/ChaosSimModuleManager.h
@@ -0,0 +1,102 @@
+// Copyright Epic Games, Inc. All Rights Reserved.
+
+#pragma once
+
+#include "CoreMinimal.h"
+#include "PhysicsInterfaceDeclaresCore.h"
+#include "Engine/World.h"
+
+#include "Chaos/SimCallbackInput.h"
+#include "Chaos/SimCallbackObject.h"
+#include "ChaosModularVehicle/ChaosSimModuleManagerAsyncCallback.h"
+
+class UModularVehicleComponent;
+class UModularVehicleBaseComponent;
+class AHUD;
+class FPhysScene_Chaos;
+
+class CHAOSMODULARVEHICLEENGINE_API FChaosSimModuleManager
+{
+public:
+	// Updated when vehicles need to recreate their physics state.
+	// Used when values tweaked while the game is running.
+	//static uint32 VehicleSetupTag;
+
+	FChaosSimModuleManager(FPhysScene* PhysScene);
+	~FChaosSimModuleManager();
+
+	static void OnPostWorldInitialization(UWorld* InWorld, const UWorld::InitializationValues);
+	static void OnWorldCleanup(UWorld* InWorld, bool bSessionEnded, bool bCleanupResources);
+	static void OnShowDebugInfo(AHUD* HUD, UCanvas* Canvas, const FDebugDisplayInfo& DisplayInfo, float& YL, float& YPos);
+
+	/** Get Physics Scene */
+	FPhysScene_Chaos& GetScene() const { return Scene; }
+
+	/**
+	 * Register a Physics vehicle for processing
+	 */
+	void AddVehicle(TWeakObjectPtr<UModularVehicleComponent> Vehicle);
+	void AddVehicle(TWeakObjectPtr<UModularVehicleBaseComponent> Vehicle);
+
+	/**
+	 * Unregister a Physics vehicle from processing
+	 */
+	void RemoveVehicle(TWeakObjectPtr<UModularVehicleComponent> Vehicle);
+	void RemoveVehicle(TWeakObjectPtr<UModularVehicleBaseComponent> Vehicle);
+
+	/**
+	 * Update vehicle tuning and other state such as input
+	 */
+	void ScenePreTick(FPhysScene* PhysScene, float DeltaTime);
+
+	/** Detach this vehicle manager from a FPhysScene (remove delegates, remove from map etc) */
+	void DetachFromPhysScene(FPhysScene* PhysScene);
+
+	/** Update simulation of registered vehicles */
+	void Update(FPhysScene* PhysScene, float DeltaTime);
+
+	/** Post update step */
+	void PostUpdate(FChaosScene* PhysScene);
+
+	void ParallelUpdateVehicles(float DeltaSeconds);
+
+	/** Find a vehicle manager from an FPhysScene */
+	static FChaosSimModuleManager* GetManagerFromScene(FPhysScene* PhysScene);
+
+protected:
+	void RegisterCallbacks(UWorld* InWorld);
+	void UnregisterCallbacks();
+
+private:
+	/** Map of physics scenes to corresponding vehicle manager */
+	static TMap<FPhysScene*, FChaosSimModuleManager*> SceneToModuleManagerMap;
+	
+	// The physics scene we belong to
+	FPhysScene_Chaos& Scene;
+
+	static bool GInitialized;
+
+	// All instanced vehicles
+	TArray<TWeakObjectPtr<UModularVehicleComponent>> GCVehicles;
+
+	// Test out new vehicle bass class
+	TArray<TWeakObjectPtr<UModularVehicleBaseComponent>> CUVehicles;
+
+	// callback delegates
+	FDelegateHandle OnPhysScenePreTickHandle;
+	FDelegateHandle OnPhysScenePostTickHandle;
+
+	static FDelegateHandle OnPostWorldInitializationHandle;
+	static FDelegateHandle OnWorldCleanupHandle;
+
+	// Async callback from the physics engine - we can run our simulation here
+	FChaosSimModuleManagerAsyncCallback* AsyncCallback;	
+	int32 Timestamp;
+	int32 SubStepCount;
+
+	// async input/output state
+	TArray<Chaos::TSimCallbackOutputHandle<FChaosSimModuleManagerAsyncOutput>> PendingOutputs;
+	Chaos::TSimCallbackOutputHandle<FChaosSimModuleManagerAsyncOutput> LatestOutput;
+
+};
+
diff --git a/Engine/Plugins/Experimental/ChaosModularVehicle/Source/ChaosModularVehicleEngine/Public/ChaosModularVehicle/ChaosSimModuleManagerAsyncCallback.h b/Engine/Plugins/Experimental/ChaosModularVehicle/Source/ChaosModularVehicleEngine/Public/ChaosModularVehicle/ChaosSimModuleManagerAsyncCallback.h
new file mode 100644
--- /dev/null
+++ b/Engine/Plugins/Experimental/ChaosModularVehicle/Source/ChaosModularVehicleEngine/Public/ChaosModularVehicle/ChaosSimModuleManagerAsyncCallback.h
@@ -0,0 +1,324 @@
+// Copyright Epic Games, Inc. All Rights Reserved.
+
+#pragma once
+
+#include "CoreMinimal.h"
+#include "PhysicsPublic.h"
+#include "Chaos/SimCallbackInput.h"
+#include "Chaos/SimCallbackObject.h"
+#include "Chaos/GeometryParticlesfwd.h"
+#include "TransmissionSystem.h"
+#include "Physics/NetworkPhysicsComponent.h"
+#include "SimModule/SimulationModuleBase.h"
+
+#include "ChaosSimModuleManagerAsyncCallback.generated.h"
+
+class UModularVehicleComponent;
+class UModularVehicleBaseComponent;
+class FGeometryCollectionPhysicsProxy;
+
+DECLARE_STATS_GROUP(TEXT("ChaosSimModuleManager"), STATGROUP_ChaosSimModuleManager, STATGROUP_Advanced);
+
+
+struct FSimModuleDebugParams
+{
+	bool EnableMultithreading = false;
+	bool EnableNetworkStateData = true;
+};
+
+enum EChaosAsyncVehicleDataType : int8
+{
+	AsyncInvalid,
+	AsyncDefault,
+};
+
+
+struct FModuleTransform
+{
+	int TransforIndex;
+	FTransform Transform;
+};
+
+/** Vehicle inputs from the player controller */
+USTRUCT()
+struct CHAOSMODULARVEHICLEENGINE_API FModularVehicleInputs
+{
+	GENERATED_USTRUCT_BODY()
+
+		FModularVehicleInputs()
+		: Steering(0.f)
+		, Throttle(0.f)
+		, Brake(0.f)
+		, Handbrake(0.f)
+		, Pitch(0.f)
+		, Roll(0.f)
+		, Yaw(0.f)
+		, Boost(0.f)
+		, Drift(0.f)
+		, Reverse(false)
+		, KeepAwake(false)
+		{}
+
+	// Steering output to physics system. Range -1...1
+	UPROPERTY()
+	float Steering;
+
+	// Accelerator output to physics system. Range 0...1
+	UPROPERTY()
+	float Throttle;
+
+	// Brake output to physics system. Range 0...1
+	UPROPERTY()
+	float Brake;
+
+	// Handbrake output to physics system. Range 0...1
+	UPROPERTY()
+	float Handbrake;
+
+	// Body Pitch output to physics system. Range -1...1
+	UPROPERTY()
+	float Pitch;
+
+	// Body Roll output to physics system. Range -1...1
+	UPROPERTY()
+	float Roll;
+
+	// Body Yaw output to physics system. Range -1...1
+	UPROPERTY()
+	float Yaw;
+
+	// Boost output to physics system. Range -1...1
+	UPROPERTY()
+	float Boost;
+
+	// Boost output to physics system. Range 0...1
+	UPROPERTY()
+	float Drift;
+
+	// Reversing state
+	UPROPERTY()
+	bool Reverse;
+
+	// Keep vehicle awake
+	UPROPERTY()
+	bool KeepAwake;
+
+};
+
+
+/** Vehicle Inputs datas that will be used in the inputs history to be applied while simulating */
+USTRUCT()
+struct CHAOSMODULARVEHICLEENGINE_API FNetworkModularVehicleInputs : public FNetworkPhysicsDatas
+{
+	GENERATED_USTRUCT_BODY()
+
+	/** List of incoming control inputs coming from the local client */
+	UPROPERTY()
+	FModularVehicleInputs VehicleInputs;
+
+	/**  Apply the datas onto the network physics component */
+	virtual void ApplyDatas(UActorComponent* NetworkComponent) const override;
+
+	/**  Build the datas from the network physics component */
+	virtual void BuildDatas(const UActorComponent* NetworkComponent) override;
+
+	/**  Serialize datas function that will be used to transfer the struct across the network */
+	bool NetSerialize(FArchive& Ar, class UPackageMap* Map, bool& bOutSuccess);
+
+	/** Interpolate the datas in between two inputs datas */
+	void InterpolateDatas(const FNetworkModularVehicleInputs& MinDatas, const FNetworkModularVehicleInputs& MaxDatas);
+};
+
+template<>
+struct TStructOpsTypeTraits<FNetworkModularVehicleInputs> : public TStructOpsTypeTraitsBase2<FNetworkModularVehicleInputs>
+{
+	enum
+	{
+		WithNetSerializer = true,
+	};
+};
+
+/** Vehicle states datas that will be used in the states history to rewind the simulation at some point inn time */
+USTRUCT()
+struct CHAOSMODULARVEHICLEENGINE_API FNetworkModularVehicleStates : public FNetworkPhysicsDatas
+{
+	GENERATED_USTRUCT_BODY()
+
+	Chaos::FModuleNetDataArray ModuleData;
+
+	/**  Apply the datas onto the network physics component */
+	virtual void ApplyDatas(UActorComponent* NetworkComponent) const override;
+
+	/**  Build the datas from the network physics component */
+	virtual void BuildDatas(const UActorComponent* NetworkComponent) override;
+
+	/**  Serialize datas function that will be used to transfer the struct across the network */
+	bool NetSerialize(FArchive& Ar, class UPackageMap* Map, bool& bOutSuccess);
+
+	/** Interpolate the datas in between two inputs datas */
+	void InterpolateDatas(const FNetworkModularVehicleStates& MinDatas, const FNetworkModularVehicleStates& MaxDatas);
+};
+
+template<>
+struct TStructOpsTypeTraits<FNetworkModularVehicleStates> : public TStructOpsTypeTraitsBase2<FNetworkModularVehicleStates>
+{
+	enum
+	{
+		WithNetSerializer = true,
+	};
+};
+
+/**
+ * Per Vehicle Output State from Physics Thread to Game Thread
+ */
+struct CHAOSMODULARVEHICLEENGINE_API FPhysicsVehicleOutput
+{
+	FPhysicsVehicleOutput()
+	{
+	}
+
+	~FPhysicsVehicleOutput()
+	{
+		Clean();
+	}
+
+	void Clean()
+	{
+		for (Chaos::FSimOutputData* Data : SimTreeOutputData)
+		{
+			delete Data;
+		}
+		SimTreeOutputData.Empty();
+	}
+
+	TArray<Chaos::FSimOutputData*> SimTreeOutputData;
+	TArray<FString> DebugStrings;
+};
+
+struct CHAOSMODULARVEHICLEENGINE_API FPhysicsModularVehicleTraits
+{
+	using InputsType = FNetworkModularVehicleInputs;
+	using StatesType = FNetworkModularVehicleStates;
+};
+
+// #TBD
+struct FGameStateInputs
+{
+	//bool IsInWater;
+	//bool IsInAir
+};
+
+
+/**
+ * Per Vehicle input State from Game Thread to Physics Thread
+ */
+struct CHAOSMODULARVEHICLEENGINE_API FPhysicsModularVehicleInputs
+{
+	FPhysicsModularVehicleInputs()
+		: TraceParams()
+		, TraceCollisionResponse()
+	{
+	}
+	mutable FNetworkModularVehicleInputs NetworkInputs;
+	mutable FCollisionQueryParams TraceParams;
+	mutable FCollisionResponseContainer TraceCollisionResponse;
+	mutable FGameStateInputs StateInputs;
+};
+
+/**
+ * Per Vehicle Input State from Game Thread to Physics Thread
+ */
+struct CHAOSMODULARVEHICLEENGINE_API FModularVehicleAsyncInput
+{
+	FModularVehicleAsyncInput(EChaosAsyncVehicleDataType InType = EChaosAsyncVehicleDataType::AsyncInvalid)
+		: Type(InType)
+		, Vehicle(nullptr)
+	{
+		Proxy = nullptr;	//indicates async/sync task not needed
+	}
+
+	virtual ~FModularVehicleAsyncInput() = default;
+
+	/**
+	* Vehicle simulation running on the Physics Thread
+	*/
+	virtual TUniquePtr<struct FModularVehicleAsyncOutput> Simulate(UWorld* World, const float DeltaSeconds, const float TotalSeconds, bool& bWakeOut) const;
+	virtual void ApplyDeferredForces() const;
+	virtual void ProcessInputs();
+
+	void SetVehicle(UModularVehicleBaseComponent* VehicleIn) { Vehicle = VehicleIn; }
+	UModularVehicleBaseComponent* GetVehicle() { return Vehicle; }
+
+	const EChaosAsyncVehicleDataType Type;
+	IPhysicsProxyBase* Proxy;
+
+	FPhysicsModularVehicleInputs PhysicsInputs;
+
+private:
+
+	UModularVehicleBaseComponent* Vehicle;
+
+};
+
+struct FChaosSimModuleManagerAsyncInput : public Chaos::FSimCallbackInput
+{
+	TArray<TUniquePtr<FModularVehicleAsyncInput>> VehicleInputs;
+
+	TWeakObjectPtr<UWorld> World;
+	int32 Timestamp = INDEX_NONE;
+
+	void Reset()
+	{
+		VehicleInputs.Reset();
+		World.Reset();
+	}
+};
+
+/**
+ * Async Output Data
+ */
+struct CHAOSMODULARVEHICLEENGINE_API FModularVehicleAsyncOutput
+{
+	const EChaosAsyncVehicleDataType Type;
+	bool bValid;	// indicates no work was done
+	FPhysicsVehicleOutput VehicleSimOutput;
+
+	FModularVehicleAsyncOutput(EChaosAsyncVehicleDataType InType = EChaosAsyncVehicleDataType::AsyncInvalid)
+		: Type(InType)
+		, bValid(false)
+	{ }
+
+	virtual ~FModularVehicleAsyncOutput()
+	{
+		VehicleSimOutput.Clean();
+	}
+};
+
+
+/**
+ * Async Output for all of the vehicles handled by this Vehicle Manager
+ */
+struct FChaosSimModuleManagerAsyncOutput : public Chaos::FSimCallbackOutput
+{
+	TArray<TUniquePtr<FModularVehicleAsyncOutput>> VehicleOutputs;
+	int32 Timestamp = INDEX_NONE;
+
+	void Reset()
+	{
+		VehicleOutputs.Reset();
+	}
+};
+
+/**
+ * Async callback from the Physics Engine where we can perform our vehicle simulation
+ */
+class CHAOSMODULARVEHICLEENGINE_API FChaosSimModuleManagerAsyncCallback : public Chaos::TSimCallbackObject<FChaosSimModuleManagerAsyncInput, FChaosSimModuleManagerAsyncOutput
+	, Chaos::ESimCallbackOptions::Presimulate | Chaos::ESimCallbackOptions::Rewind | Chaos::ESimCallbackOptions::ContactModification>
+{
+public:
+	virtual FName GetFNameForStatId() const override;
+private:
+	virtual void ProcessInputs_Internal(int32 PhysicsStep) override;
+	virtual void OnPreSimulate_Internal() override;
+	virtual void OnContactModification_Internal(Chaos::FCollisionContactModifier& Modifications) override;
+};
diff --git a/Engine/Plugins/Experimental/ChaosModularVehicle/Source/ChaosModularVehicleEngine/Public/ChaosModularVehicle/ModularVehicleActor.h b/Engine/Plugins/Experimental/ChaosModularVehicle/Source/ChaosModularVehicleEngine/Public/ChaosModularVehicle/ModularVehicleActor.h
new file mode 100644
--- /dev/null
+++ b/Engine/Plugins/Experimental/ChaosModularVehicle/Source/ChaosModularVehicleEngine/Public/ChaosModularVehicle/ModularVehicleActor.h
@@ -0,0 +1,39 @@
+// Copyright Epic Games, Inc. All Rights Reserved.
+
+#pragma once
+
+#include "CoreMinimal.h"
+#include "GameFramework/Actor.h"
+#include "UObject/ObjectMacros.h"
+#include "UObject/ObjectPtr.h"
+#include "ChaosModularVehicle/ModularVehicleComponent.h"
+
+#include "ModularVehicleActor.generated.h"
+
+class ASimModuleActor;
+class UGeometryCollectionDebugDrawComponent;
+
+
+UCLASS()
+class CHAOSMODULARVEHICLEENGINE_API AModularVehicleActor: public AActor
+{
+	GENERATED_UCLASS_BODY()
+
+protected:
+	//virtual void BeginPlay() override;
+
+public:
+
+	/* Game state callback */
+	void TickSimulation(float DeltaTime);
+	virtual void Tick(float DeltaSeconds) override;
+
+	virtual void PostInitializeComponents() override;
+
+	/* ModularVehicleComponent */
+	UPROPERTY(VisibleAnywhere, BlueprintReadOnly, Category = Destruction, meta = (ExposeFunctionCategories = "Components|ModularVehicle", AllowPrivateAccess = "true"))
+	TObjectPtr<UModularVehicleComponent> ModularVehicleComponent;
+
+	UModularVehicleComponent* GetModularVehicleComponent() const { return ModularVehicleComponent; }
+
+};
diff --git a/Engine/Plugins/Experimental/ChaosModularVehicle/Source/ChaosModularVehicleEngine/Public/ChaosModularVehicle/ModularVehicleBaseComponent.h b/Engine/Plugins/Experimental/ChaosModularVehicle/Source/ChaosModularVehicleEngine/Public/ChaosModularVehicle/ModularVehicleBaseComponent.h
new file mode 100644
--- /dev/null
+++ b/Engine/Plugins/Experimental/ChaosModularVehicle/Source/ChaosModularVehicleEngine/Public/ChaosModularVehicle/ModularVehicleBaseComponent.h
@@ -0,0 +1,550 @@
+// Copyright Epic Games, Inc. All Rights Reserved.
+
+#pragma once
+
+#include "ChaosModularVehicle/ChaosSimModuleManagerAsyncCallback.h"
+#include "SimModule/SimModulesInclude.h"
+#include "ChaosModularVehicle/ModularVehicleInputRate.h"
+#include "ChaosModularVehicle/ModularVehicleSimulationCU.h"
+#include "Chaos/ParticleHandleFwd.h"
+#include "Components/PrimitiveComponent.h"
+#include "Components/ActorComponent.h"
+#include "Containers/Map.h"
+#include "UObject/ObjectKey.h"
+
+#include "ModularVehicleBaseComponent.generated.h"
+
+DECLARE_LOG_CATEGORY_EXTERN(LogModularBase, Log, All);
+
+struct FModularVehicleAsyncInput;
+struct FChaosSimModuleManagerAsyncOutput;
+class UClusterUnionComponent;
+namespace Chaos
+{
+	class FSimTreeUpdates;
+}
+
+USTRUCT()
+struct FVehicleComponentData
+{
+	GENERATED_BODY()
+
+	int Guid = -1;
+};
+
+/** Additional replicated state */
+USTRUCT()
+struct CHAOSMODULARVEHICLEENGINE_API FModularReplicatedState : public FModularVehicleInputs
+{
+	GENERATED_USTRUCT_BODY()
+
+	FModularReplicatedState() : FModularVehicleInputs()
+	{
+	}
+
+};
+
+/** Input Options */
+UENUM()
+enum class EFunctionType : uint8
+{
+	LinearFunction = 0,
+	SquaredFunction,
+	CustomCurve
+};
+
+USTRUCT()
+struct CHAOSMODULARVEHICLEENGINE_API FModularVehicleInputRateConfig
+{
+	GENERATED_USTRUCT_BODY()
+
+	/**
+		* Rate at which the input value rises
+		*/
+	UPROPERTY(EditAnywhere, Category = VehicleInputRate)
+	float RiseRate;
+
+	/**
+	 * Rate at which the input value falls
+	 */
+	UPROPERTY(EditAnywhere, Category = VehicleInputRate)
+	float FallRate;
+
+	/**
+	 * Controller input curve, various predefined options, linear, squared, or user can specify a custom curve function
+	 */
+	UPROPERTY(EditAnywhere, Category = VehicleInputRate)
+	EFunctionType InputCurveFunction;
+
+	/**
+	 * Controller input curve - should be a normalized float curve, i.e. time from 0 to 1 and values between 0 and 1
+	 * This curve is only sued if the InputCurveFunction above is set to CustomCurve
+	 */
+	UPROPERTY(EditAnywhere, Category = VehicleInputRate)
+	FRuntimeFloatCurve UserCurve;
+
+	FModularVehicleInputRateConfig() : RiseRate(5.0f), FallRate(5.0f), InputCurveFunction(EFunctionType::LinearFunction) { }
+
+	/** Change an output value using max rise and fall rates */
+	float InterpInputValue(float DeltaTime, float CurrentValue, float NewValue) const
+	{
+		const float DeltaValue = NewValue - CurrentValue;
+
+		// We are "rising" when DeltaValue has the same sign as CurrentValue (i.e. delta causes an absolute magnitude gain)
+		// OR we were at 0 before, and our delta is no longer 0.
+		const bool bRising = ((DeltaValue > 0.0f) == (CurrentValue > 0.0f)) ||
+			((DeltaValue != 0.f) && (CurrentValue == 0.f));
+
+		const float MaxDeltaValue = DeltaTime * (bRising ? RiseRate : FallRate);
+		const float ClampedDeltaValue = FMath::Clamp(DeltaValue, -MaxDeltaValue, MaxDeltaValue);
+		return CurrentValue + ClampedDeltaValue;
+	}
+
+	float CalcControlFunction(float InputValue)
+	{
+		// user defined curve
+
+		// else use option from drop down list
+		switch (InputCurveFunction)
+		{
+		case EFunctionType::CustomCurve:
+		{
+			if (UserCurve.GetRichCurveConst() && !UserCurve.GetRichCurveConst()->IsEmpty())
+			{
+				float Output = FMath::Clamp(UserCurve.GetRichCurveConst()->Eval(FMath::Abs(InputValue)), 0.0f, 1.0f);
+				return (InputValue < 0.f) ? -Output : Output;
+			}
+			else
+			{
+				return InputValue;
+			}
+		}
+		break;
+		case EFunctionType::SquaredFunction:
+		{
+			return (InputValue < 0.f) ? -InputValue * InputValue : InputValue * InputValue;
+		}
+		break;
+
+		case EFunctionType::LinearFunction:
+		default:
+		{
+			return InputValue;
+		}
+		break;
+
+		}
+
+	}
+};
+
+USTRUCT()
+struct CHAOSMODULARVEHICLEENGINE_API FConstructionData
+{
+	GENERATED_USTRUCT_BODY()
+
+	UPROPERTY()
+	UPrimitiveComponent* Component = nullptr;
+
+	UPROPERTY()
+	int32 ConstructionIndex = INDEX_NONE;
+};
+
+
+UCLASS(ClassGroup = (Physics), meta = (BlueprintSpawnableComponent), hidecategories = (PlanarMovement, "Components|Movement|Planar", Activation, "Components|Activation"))
+class CHAOSMODULARVEHICLEENGINE_API UModularVehicleBaseComponent : public UPawnMovementComponent
+{
+	GENERATED_UCLASS_BODY()
+
+	~UModularVehicleBaseComponent();
+
+	friend struct FModularVehicleAsyncInput;
+	friend struct FModularVehicleAsyncOutput;
+
+	friend struct FNetworkModularVehicleInputs;
+	friend struct FNetworkModularVehicleStates;
+
+	friend class FModularVehicleManager;
+	friend class FChaosSimModuleManagerAsyncCallback;
+
+	friend class FModularVehicleBuilder;
+public:
+	APlayerController* GetPlayerController() const;
+
+	virtual void GetLifetimeReplicatedProps(TArray<FLifetimeProperty>& OutLifetimeProps) const override;
+
+	virtual bool ShouldCreatePhysicsState() const override { return true; }
+	virtual void OnCreatePhysicsState() override;
+	virtual void OnDestroyPhysicsState() override;
+
+	virtual void BeginPlay() override;
+	virtual void TickComponent(float DeltaTime, ELevelTick TickType, FActorComponentTickFunction* ThisTickFunction) override;
+
+	void CreateAssociatedSimComponents(const UPrimitiveComponent* AttachedComponent, int ParentIndex, int TransformIndex, Chaos::FSimTreeUpdates& TreeUpdatesOut);
+
+	void CreateConstraint(Chaos::ISimulationModuleBase* NewModule);
+	void DestroyConstraint(int ConstraintIndex);
+	void DestroyAllConstraints();
+	void EnableConstraint(int ConstraintIndex, bool bEnabled);
+
+	void PreTickGT(float DeltaTime);
+	void UpdateState(float DeltaTime);
+	TUniquePtr<FModularVehicleAsyncInput> SetCurrentAsyncData(int32 InputIdx, FChaosSimModuleManagerAsyncOutput* CurOutput, FChaosSimModuleManagerAsyncOutput* NextOutput, float Alpha, int32 VehicleManagerTimestamp);
+
+	void SetClusterComponent(UClusterUnionComponent* InPhysicalComponent);
+	void ParallelUpdate(float DeltaTime);
+	void Update(float DeltaTime);
+	void FinalizeSimCallbackData(FChaosSimModuleManagerAsyncInput& Input);
+
+	/** handle stand-alone and networked mode control inputs */
+	void ProcessControls(float DeltaTime);
+
+	void ShowDebugInfo(AHUD* HUD, UCanvas* Canvas, const FDebugDisplayInfo& DisplayInfo, float& YL, float& YPos);
+
+	TUniquePtr<FPhysicsVehicleOutput>& PhysicsVehicleOutput()
+	{
+		return PVehicleOutput;
+	}
+
+	FORCEINLINE const FTransform& GetComponentTransform() const;
+
+	//UPROPERTY(EditAnywhere, Category = VehicleInput, AdvancedDisplay)
+	TArray<FModularVehicleInputRate> InputInterpolationRates;
+
+	/** Use to naturally decelerate linear velocity of objects */
+	UPROPERTY(EditAnywhere, Category = "Game|Components|ModularVehicle")
+	float LinearDamping;
+
+	/** Use to naturally decelerate angular velocity of objects */
+	UPROPERTY(EditAnywhere, Category = "Game|Components|ModularVehicle")
+	float AngularDamping;
+
+	UPROPERTY(EditAnywhere, Category = "Game|Components|ModularVehicle")
+	struct FCollisionResponseContainer SuspensionTraceCollisionResponses;
+
+	UPROPERTY(EditAnywhere, Category = "Game|Components|ModularVehicle")
+	bool bSuspensionTraceComplex;
+
+	UPROPERTY(EditAnywhere, Category = "Game|Components|ModularVehicle")
+	bool bKeepVehicleAwake;
+
+	/** Adds any associated simulation components to the ModularVehicleSimulation */
+	UFUNCTION()
+	void AddComponentToSimulation(UPrimitiveComponent* Component, const TArray<FClusterUnionBoneData>& BonesData, const TArray<FClusterUnionBoneData>& RemovedBoneIDs, bool bIsNew);
+
+	/** Removes any associated simulation components from the ModularVehicleSimulation */
+	UFUNCTION()
+	void RemoveComponentFromSimulation(UPrimitiveComponent* Component, const TArray<FClusterUnionBoneData>& RemovedBonesData);
+
+	// CONTROLS
+	// 
+	/** Set the user input for the vehicle throttle [range 0 to 1] */
+	UFUNCTION(BlueprintCallable, Category = "Game|Components|ModularVehicle")
+	void SetThrottleInput(float Throttle);
+
+	/** Set the user input for the vehicle boost [range 0 to 1] */
+	UFUNCTION(BlueprintCallable, Category = "Game|Components|ModularVehicle")
+	void SetBoostInput(float Boost);
+
+	/** Set the user input for the vehicle drift [range 0 to 1] */
+	UFUNCTION(BlueprintCallable, Category = "Game|Components|ModularVehicle")
+	void SetDriftInput(float Drift);
+
+	/** Increase the vehicle throttle position [throttle range normalized 0 to 1] */
+	UFUNCTION(BlueprintCallable, Category = "Game|Components|ModularVehicle")
+	void IncreaseThrottleInput(float ThrottleDelta);
+
+	/** Decrease the vehicle throttle position  [throttle range normalized 0 to 1] */
+	UFUNCTION(BlueprintCallable, Category = "Game|Components|ModularVehicle")
+	void DecreaseThrottleInput(float ThrottleDelta);
+
+	/** Set the user input for the vehicle Brake [range 0 to 1] */
+	UFUNCTION(BlueprintCallable, Category = "Game|Components|ModularVehicle")
+	void SetBrakeInput(float Brake);
+
+	/** Set the user input for the vehicle steering [range -1 to 1] */
+	UFUNCTION(BlueprintCallable, Category = "Game|Components|ModularVehicle")
+	void SetSteeringInput(float Steering);
+
+	/** Set the user input for the vehicle pitch [range -1 to 1] */
+	UFUNCTION(BlueprintCallable, Category = "Game|Components|ModularVehicle")
+	void SetPitchInput(float Pitch);
+
+	/** Set the user input for the vehicle roll [range -1 to 1] */
+	UFUNCTION(BlueprintCallable, Category = "Game|Components|ModularVehicle")
+	void SetRollInput(float Roll);
+
+	/** Set the user input for the vehicle yaw [range -1 to 1] */
+	UFUNCTION(BlueprintCallable, Category = "Game|Components|ModularVehicle")
+	void SetYawInput(float Yaw);
+
+	/** Set the user input for handbrake */
+	UFUNCTION(BlueprintCallable, Category = "Game|Components|ModularVehicle")
+	void SetHandbrakeInput(float Handbrake);
+
+	UFUNCTION(BlueprintCallable, Category = "Game|Components|ModularVehicle")
+	void SetReverseInput(bool Reverse);
+
+	/** Set the gear directly */
+	UFUNCTION(BlueprintCallable, Category = "Game|Components|ModularVehicle")
+	void SetGearInput(int32 Gear);
+
+	UFUNCTION(BlueprintCallable, Category = "Game|Components|ModularVehicle")
+	int32 GetCurrentGear();
+
+	UFUNCTION(BlueprintCallable, Category = "Game|Components|ModularVehicle")
+	bool IsReversing();
+
+	UFUNCTION(BlueprintCallable, Category = "Game|Components|ModularVehicle")
+	void AddActorsToIgnore(TArray<AActor*>& ActorsIn);
+
+	UFUNCTION(BlueprintCallable, Category = "Game|Components|ModularVehicle")
+	void RemoveActorsToIgnore(TArray<AActor*>& ActorsIn);
+
+	// Bypass the need for a controller in order for the controls to be processed.
+	UPROPERTY(EditAnywhere, Category = "Game|Components|ModularVehicle")
+	bool bRequiresControllerForInputs;
+
+	/** Grab nearby components and add them to the cluster union representing the vehicle */
+	UPROPERTY(EditAnywhere, Category = "Game|Components|ModularVehicle")
+	bool bAutoAddComponentsFromWorld;
+
+	/** The size of the overlap box testing for nearby components in the world  */
+	UPROPERTY(EditAnywhere, Category = "Game|Components|ModularVehicle", meta = (EditCondition = "bAutoAddComponentsToCluster"))
+	FVector AutoAddOverlappingBoxSize;
+
+	UPROPERTY(EditAnywhere, Category = "Game|Components|ModularVehicle", meta = (EditCondition = "bAutoAddComponentsToCluster"))
+	int32 DelayClusteringCount;
+
+	/*** Map simulation component to our vehicle setup data */
+	TMap<TObjectKey<UPrimitiveComponent>, FVehicleComponentData> ComponentToPhysicsObjects;
+
+	UClusterUnionComponent* ClusterUnionComponent;
+
+	/** Set all channels to the specified response - for wheel raycasts */
+	void SetWheelTraceAllChannels(ECollisionResponse NewResponse)
+	{
+		SuspensionTraceCollisionResponses.SetAllChannels(NewResponse);
+	}
+
+	/** Set the response of this body to the supplied settings - for wheel raycasts */
+	void SetWheelTraceResponseToChannel(ECollisionChannel Channel, ECollisionResponse NewResponse)
+	{
+		SuspensionTraceCollisionResponses.SetResponse(Channel, NewResponse);
+	}
+
+protected:
+
+	void CreateVehicleSim();
+	void DestroyVehicleSim();
+	void UpdatePhysicalProperties();
+	void AddOverlappingComponentsToCluster();
+	void AddGeometryCollectionsFromOwnedActor();
+
+	void ActionTreeUpdates(Chaos::FSimTreeUpdates* NextTreeUpdates);
+
+	void SetCurrentAsyncDataInternal(FModularVehicleAsyncInput* CurInput, int32 InputIdx, FChaosSimModuleManagerAsyncOutput* CurOutput, FChaosSimModuleManagerAsyncOutput* NextOutput, float Alpha, int32 VehicleManagerTimestamp);
+
+	IPhysicsProxyBase* GetPhysicsProxy() const;
+
+	int32 FindComponentAddOrder(UPrimitiveComponent* InComponent);
+	bool FindAndRemoveNextPendingUpdate(int32 NextIndex, Chaos::FSimTreeUpdates* OutData);
+
+	// replicated state of vehicle 
+	UPROPERTY(Transient, Replicated)
+	FModularReplicatedState ReplicatedState;
+
+	// What the player has the steering set to. Range -1...1
+	UPROPERTY(Transient)
+	float RawSteeringInput;
+
+	// What the player has the accelerator set to. Range -1...1
+	UPROPERTY(Transient)
+	float RawThrottleInput;
+
+	// What the player has the brake set to. Range -1...1
+	UPROPERTY(Transient)
+	float RawBrakeInput;
+
+	// What the player has the brake set to. Range -1...1
+	UPROPERTY(Transient)
+	float RawHandbrakeInput;
+
+	// What the player has the clutch set to. Range -1...1
+	UPROPERTY(Transient)
+	float RawClutchInput;
+
+	// What the player has the pitch set to. Range -1...1
+	UPROPERTY(Transient)
+	float RawPitchInput;
+
+	// What the player has the roll set to. Range -1...1
+	UPROPERTY(Transient)
+	float RawRollInput;
+
+	// What the player has the yaw set to. Range -1...1
+	UPROPERTY(Transient)
+	float RawYawInput;
+
+	// What the player has the yaw set to. Range -1...1
+	UPROPERTY(Transient)
+	float RawBoostInput;
+
+	// What the player has the yaw set to. Range -1...1
+	UPROPERTY(Transient)
+	float RawDriftInput;
+
+	// latest gear selected
+	UPROPERTY(Transient)
+	int32 RawGearInput;
+
+	// reverse direction enbaled
+	UPROPERTY(Transient)
+	bool RawReverseInput;
+
+	// Steering output to physics system. Range -1...1
+	UPROPERTY(Transient)
+	float SteeringInput;
+
+	// Accelerator output to physics system. Range 0...1
+	UPROPERTY(Transient)
+	float ThrottleInput;
+
+	// Brake output to physics system. Range 0...1
+	UPROPERTY(Transient)
+	float BrakeInput;
+
+	// Handbrake output to physics system. Range 0...1
+	UPROPERTY(Transient)
+	float HandbrakeInput;
+
+	// Clutch output to physics system. Range 0...1
+	UPROPERTY(Transient)
+	float ClutchInput;
+
+	// Body Pitch output to physics system. Range -1...1
+	UPROPERTY(Transient)
+	float PitchInput;
+
+	// Body Roll output to physics system. Range -1...1
+	UPROPERTY(Transient)
+	float RollInput;
+
+	// Body Yaw output to physics system. Range -1...1
+	UPROPERTY(Transient)
+	float YawInput;
+
+	// Boost output to physics system. Range 0...1
+	UPROPERTY(Transient)
+	float BoostInput;
+
+	// Boost output to physics system. Range 0...1
+	UPROPERTY(Transient)
+	float DriftInput;
+
+	// Reverse state
+	UPROPERTY(Transient)
+	bool ReverseInput;
+
+	// The currently selected gear
+	UPROPERTY(Transient)
+	int32 CurrentGear;
+
+	// The engine RPM
+	UPROPERTY(Transient)
+	float EngineRPM;
+
+	// The engine Torque
+	UPROPERTY(Transient)
+	float EngineTorque;
+
+	UPROPERTY()
+	TObjectPtr<UNetworkPhysicsComponent> NetworkPhysicsComponent = nullptr;
+
+public:
+
+	// Rate at which input throttle can rise and fall
+	UPROPERTY(EditAnywhere, Category = VehicleInput, AdvancedDisplay)
+	FModularVehicleInputRateConfig ThrottleInputRate;
+
+	// Rate at which input brake can rise and fall
+	UPROPERTY(EditAnywhere, Category = VehicleInput, AdvancedDisplay)
+	FModularVehicleInputRateConfig BrakeInputRate;
+
+	// Rate at which input steering can rise and fall
+	UPROPERTY(EditAnywhere, Category = VehicleInput, AdvancedDisplay)
+	FModularVehicleInputRateConfig SteeringInputRate;
+
+	// Rate at which input handbrake can rise and fall
+	UPROPERTY(EditAnywhere, Category = VehicleInput, AdvancedDisplay)
+	FModularVehicleInputRateConfig HandbrakeInputRate;
+
+	// Rate at which input pitch can rise and fall
+	UPROPERTY(EditAnywhere, Category = VehicleInput, AdvancedDisplay)
+	FModularVehicleInputRateConfig PitchInputRate;
+
+	// Rate at which input roll can rise and fall
+	UPROPERTY(EditAnywhere, Category = VehicleInput, AdvancedDisplay)
+	FModularVehicleInputRateConfig RollInputRate;
+
+	// Rate at which input yaw can rise and fall
+	UPROPERTY(EditAnywhere, Category = VehicleInput, AdvancedDisplay)
+	FModularVehicleInputRateConfig YawInputRate;
+
+	// Rate at which input can rise and fall
+	UPROPERTY(EditAnywhere, Category = VehicleInput, AdvancedDisplay)
+	FModularVehicleInputRateConfig BoostInputRate;
+
+	// Rate at which input can rise and fall
+	UPROPERTY(EditAnywhere, Category = VehicleInput, AdvancedDisplay)
+	FModularVehicleInputRateConfig DriftInputRate;
+
+	UPROPERTY(Transient, Replicated)
+	TArray<FConstructionData> ConstructionDatas;
+
+	/** Pass current state to server */
+	UFUNCTION(reliable, server, WithValidation)
+	void ServerUpdateState(float InSteeringInput, float InThrottleInput, float InBrakeInput
+		, float InHandbrakeInput, int32 InCurrentGear, float InRollInput, float InPitchInput
+		, float InYawInput, float InBoostInput, float InDriftInput, bool InReverseInput);
+
+	TArray<AActor*> ActorsToIgnore;
+	EChaosAsyncVehicleDataType CurAsyncType;
+	FModularVehicleAsyncInput* CurAsyncInput;
+	struct FModularVehicleAsyncOutput* CurAsyncOutput;
+	struct FModularVehicleAsyncOutput* NextAsyncOutput;
+	float OutputInterpAlpha;
+
+	struct FAsyncOutputWrapper
+	{
+		int32 Idx;
+		int32 Timestamp;
+
+		FAsyncOutputWrapper()
+			: Idx(INDEX_NONE)
+			, Timestamp(INDEX_NONE)
+		{
+		}
+	};
+
+	TArray<FAsyncOutputWrapper> OutputsWaitingOn;
+	TUniquePtr<FPhysicsVehicleOutput> PVehicleOutput;	/* physics simulation data output from the async physics thread */
+	TUniquePtr<FModularVehicleSimulationCU> VehicleSimulationPT;	/* simulation code running on the physics thread async callback */
+
+private:
+
+	int NextTransformIndex = 0; // is there a better way, getting from size of map/array?
+	UPrimitiveComponent* MyComponent = nullptr;
+
+	bool bUsingNetworkPhysicsPrediction = false;
+	float PrevSteeringInput = 0.0f;
+
+	int32 LastComponentAddIndex = INDEX_NONE;
+	TMap<TObjectKey<UPrimitiveComponent>, Chaos::FSimTreeUpdates> PendingTreeUpdates;
+
+	int32 NextConstructionIndex = 0;
+
+	TArray<FPhysicsConstraintHandle> ConstraintHandles;
+	int32 ClusteringCount = 0;
+};
+
diff --git a/Engine/Plugins/Experimental/ChaosModularVehicle/Source/ChaosModularVehicleEngine/Public/ChaosModularVehicle/ModularVehicleComponent.h b/Engine/Plugins/Experimental/ChaosModularVehicle/Source/ChaosModularVehicleEngine/Public/ChaosModularVehicle/ModularVehicleComponent.h
new file mode 100644
--- /dev/null
+++ b/Engine/Plugins/Experimental/ChaosModularVehicle/Source/ChaosModularVehicleEngine/Public/ChaosModularVehicle/ModularVehicleComponent.h
@@ -0,0 +1,391 @@
+// Copyright Epic Games, Inc. All Rights Reserved.
+#pragma once
+
+#include "ChaosModularVehicle/ModularVehicleAsset.h"
+#include "ChaosModularVehicle/ModularVehicleSimulationGC.h"
+#include "Components/MeshComponent.h"
+#include "UObject/ObjectMacros.h"
+#include "Curves/CurveFloat.h"
+
+#include "GeometryCollection/GeometryCollectionComponent.h"
+#include "ChaosModularVehicle/ChaosSimModuleManagerAsyncCallback.h"
+#include "SimModule/SimModuleTree.h"
+#include "SimModule/SimModulesInclude.h"
+#include "ChaosModularVehicle/ModularVehicleInputRate.h"
+
+#include "ModularVehicleComponent.generated.h"
+
+DECLARE_LOG_CATEGORY_EXTERN(LogModularVehicle, Log, All);
+
+namespace Chaos
+{ 
+	class FWheelSimModule;
+}
+
+struct FModularVehicleAsyncInput;
+struct FChaosSimModuleManagerAsyncOutput;
+class AHUD;
+class FDebugDisplayInfo;
+class UCanvas;
+
+USTRUCT()
+struct CHAOSMODULARVEHICLEENGINE_API FModularVehicleReplicatedState
+{
+	GENERATED_USTRUCT_BODY()
+
+	FModularVehicleReplicatedState()
+	{
+		Controls.Init(0.0f, EModularVehicleInputType::Max);
+		PrevControls.Init(0.0f, EModularVehicleInputType::Max);
+	}
+
+	TArray<float> Controls;
+	TArray<float> PrevControls;
+};
+
+
+/**
+ * Replicated data for a modular vehicle when bEnableReplication is true for
+ * that component. See UpdateRepData/ProcessRepData
+ */
+USTRUCT()
+struct FModularVehicleRepData
+{
+	GENERATED_BODY()
+
+	struct FControlChanges
+	{
+		FControlChanges() : ID(0), Value(0.0f) {}
+		FControlChanges(int InID, float InValue) : ID(InID), Value(InValue) {}
+
+		friend FArchive& operator<<(FArchive& Ar, FControlChanges& objToSerialize)
+		{
+			Ar << objToSerialize.ID;
+			Ar << objToSerialize.Value;
+			return Ar;
+		}
+
+		int ID;
+		float Value;
+	};
+
+	FModularVehicleRepData()
+		: Version(0), ServerFrame(0)
+	{
+
+	}
+
+	// Version counter, every write to the rep data is a new state so Identical only references this version
+	// as there's no reason to compare the Poses array.
+	int32 Version;
+
+	// For Network Prediction Mode we require the frame number on the server when the data was gathered
+	int32 ServerFrame;
+
+	FModularVehicleReplicatedState Data;
+
+	// Just test version to skip having to traverse the whole pose array for replication
+	bool Identical(const FModularVehicleRepData* Other, uint32 PortFlags) const
+	{
+		return Other && (Version == Other->Version);
+	}
+
+	bool NetSerialize(FArchive& Ar, class UPackageMap* Map, bool& bOutSuccess)
+	{
+		bOutSuccess = true;
+
+		Ar << Version;
+		Ar << ServerFrame;
+
+		TArray<FControlChanges> ControlInputChanges;
+		
+		if (Ar.IsSaving())
+		{
+			for (int I = 0; I < Data.Controls.Num(); I++)
+			{
+				// TODO: isn't working with this test, not sure why??	
+				//if (Data.Controls[I] != Data.PrevControls[I])
+				{
+					Data.PrevControls[I] = Data.Controls[I];
+					ControlInputChanges.Add(FControlChanges(I, Data.Controls[I]));
+				}
+			}
+		}
+
+		Ar << ControlInputChanges;
+
+		if (Ar.IsLoading())
+		{
+			for (FControlChanges Change : ControlInputChanges)
+			{
+				Data.PrevControls[Change.ID] = Data.Controls[Change.ID];
+				Data.Controls[Change.ID] = Change.Value;
+			}
+		}
+
+		return bOutSuccess;
+	}
+
+};
+
+template<>
+struct TStructOpsTypeTraits<FModularVehicleRepData> : public TStructOpsTypeTraitsBase2<FModularVehicleRepData>
+{
+	enum
+	{
+		WithNetSerializer = true,
+		WithIdentical = true,
+	};
+};
+
+
+/**
+*	ModularVehicleComponent
+*/
+UCLASS(ClassGroup = (Physics), meta = (BlueprintSpawnableComponent), hidecategories = (PlanarMovement, "Components|Movement|Planar", Activation, "Components|Activation"))
+class CHAOSMODULARVEHICLEENGINE_API UModularVehicleComponent : public UGeometryCollectionComponent
+{
+	friend struct FModularVehicleDefaultAsyncInput;
+
+	GENERATED_UCLASS_BODY()
+
+public:
+	~UModularVehicleComponent();
+	virtual void TickComponent(float DeltaTime, enum ELevelTick TickType, FActorComponentTickFunction *ThisTickFunction) override;
+
+	virtual void GetLifetimeReplicatedProps(TArray<FLifetimeProperty>& OutLifetimeProps) const override;
+	virtual void AsyncPhysicsTickComponent(float DeltaTime, float SimTime) override;
+
+	/** Return true if it's suitable to create a physics representation of the vehicle at this time */
+	virtual bool ShouldCreatePhysicsState() const override;
+
+	virtual void OnCreatePhysicsState() override;
+	virtual void OnDestroyPhysicsState() override;
+
+	/** handle stand-alone and networked mode control inputs */
+	void ProcessControls(float DeltaTime);
+
+	/** Updates the vehicle tuning and other state such as user input. */
+	virtual void PreTickGT(float DeltaTime);
+
+	void CreateVehicleSim();
+	void DestroyVehicleSim();
+
+	TUniquePtr<Chaos::FSimModuleTree> GenerateSimTree();
+
+
+	TUniquePtr<FPhysicsVehicleOutput>& PhysicsVehicleOutput()
+	{
+		return PVehicleOutput;
+	}
+
+	TUniquePtr<FModularVehicleAsyncInput> SetCurrentAsyncData(int32 InputIdx, FChaosSimModuleManagerAsyncOutput* CurOutput, FChaosSimModuleManagerAsyncOutput* NextOutput, float Alpha, int32 VehicleManagerTimestamp);
+
+	void SetCurrentAsyncDataInternal(FModularVehicleAsyncInput* CurInput, int32 InputIdx, FChaosSimModuleManagerAsyncOutput* CurOutput, FChaosSimModuleManagerAsyncOutput* NextOutput, float Alpha, int32 VehicleManagerTimestamp);
+
+	virtual void Update(float DeltaTime);
+	virtual void PostUpdate();
+
+	virtual void ResetVehicleState();
+
+	// Get output data from Physics Thread
+	virtual void ParallelUpdate(float DeltaSeconds);
+	void FinalizeSimCallbackData(FChaosSimModuleManagerAsyncInput& Input);
+	void ShowDebugInfo(AHUD* HUD, UCanvas* Canvas, const FDebugDisplayInfo& DisplayInfo, float& YL, float& YPos);
+
+protected:
+
+	// accumulator for RB replication errors 
+	float AngErrorAccumulator;
+
+	// What the player has the steering set to. Range -1...1
+	UPROPERTY(Transient)
+	float RawSteeringInput;
+
+	// What the player has the accelerator set to. Range -1...1
+	UPROPERTY(Transient)
+	float RawThrottleInput;
+
+	// What the player has the brake set to. Range -1...1
+	UPROPERTY(Transient)
+	float RawBrakeInput;
+
+	// What the player has the brake set to. Range -1...1
+	UPROPERTY(Transient)
+	float RawHandbrakeInput;
+
+	// What the player has the clutch set to. Range -1...1
+	UPROPERTY(Transient)
+	float RawClutchInput;
+
+	// What the player has the pitch set to. Range -1...1
+	UPROPERTY(Transient)
+	float RawPitchInput;
+
+	// What the player has the roll set to. Range -1...1
+	UPROPERTY(Transient)
+	float RawRollInput;
+
+	// What the player has the yaw set to. Range -1...1
+	UPROPERTY(Transient)
+	float RawYawInput;
+
+	// latest gear selected
+	UPROPERTY(Transient)
+	int RawGearInput;
+
+	// Steering output to physics system. Range -1...1
+	UPROPERTY(Transient)
+	float SteeringInput;
+
+	// Accelerator output to physics system. Range 0...1
+	UPROPERTY(Transient)
+	float ThrottleInput;
+
+	// Brake output to physics system. Range 0...1
+	UPROPERTY(Transient)
+	float BrakeInput;
+
+	// Clutch output to physics system. Range 0...1
+	UPROPERTY(Transient)
+	float ClutchInput;
+
+	// Body Pitch output to physics system. Range -1...1
+	UPROPERTY(Transient)
+	float PitchInput;
+
+	// Body Roll output to physics system. Range -1...1
+	UPROPERTY(Transient)
+	float RollInput;
+
+	// Body Yaw output to physics system. Range -1...1
+	UPROPERTY(Transient)
+	float YawInput;
+
+	// Handbrake output to physics system. Range 0...1
+	UPROPERTY(Transient)
+	float HandbrakeInput;
+
+	// Bypass the need for a controller in order for the controls to be processed.
+	UPROPERTY(EditAnywhere, Category = VehicleInput)
+	bool bRequiresControllerForInputs;
+
+	// How much to press the brake when the player has release throttle
+	UPROPERTY(EditAnywhere, Category = VehicleInput)
+	float IdleBrakeInput;
+
+	// Auto-brake when absolute vehicle forward speed is less than this (cm/s)
+	UPROPERTY(EditAnywhere, Category = VehicleInput)
+	float StopThreshold;
+
+	// Auto-brake when vehicle forward speed is opposite of player input by at least this much (cm/s)
+	UPROPERTY(EditAnywhere, Category = VehicleInput)
+	float WrongDirectionThreshold;
+
+	UPROPERTY(EditAnywhere, Category = VehicleInput, AdvancedDisplay)
+	TArray<FModularVehicleInputRate> InputInterpolationRates;
+
+	/** Set the user input for the vehicle throttle [range 0 to 1] */
+	UFUNCTION(BlueprintCallable, Category = "Game|Components|ChaosVehicleMovement")
+	void SetThrottleInput(float Throttle);
+
+	/** Increase the vehicle throttle position [throttle range normalized 0 to 1] */
+	UFUNCTION(BlueprintCallable, Category = "Game|Components|ChaosVehicleMovement")
+	void IncreaseThrottleInput(float ThrottleDelta);
+
+	/** Decrease the vehicle throttle position  [throttle range normalized 0 to 1] */
+	UFUNCTION(BlueprintCallable, Category = "Game|Components|ChaosVehicleMovement")
+	void DecreaseThrottleInput(float ThrottleDelta);
+
+	/** Set the user input for the vehicle Brake [range 0 to 1] */
+	UFUNCTION(BlueprintCallable, Category = "Game|Components|ChaosVehicleMovement")
+	void SetBrakeInput(float Brake);
+
+	/** Set the user input for the vehicle steering [range -1 to 1] */
+	UFUNCTION(BlueprintCallable, Category = "Game|Components|ChaosVehicleMovement")
+	void SetSteeringInput(float Steering);
+
+	/** Set the user input for the vehicle pitch [range -1 to 1] */
+	UFUNCTION(BlueprintCallable, Category = "Game|Components|ChaosVehicleMovement")
+	void SetPitchInput(float Pitch);
+
+	/** Set the user input for the vehicle roll [range -1 to 1] */
+	UFUNCTION(BlueprintCallable, Category = "Game|Components|ChaosVehicleMovement")
+	void SetRollInput(float Roll);
+
+	/** Set the user input for the vehicle yaw [range -1 to 1] */
+	UFUNCTION(BlueprintCallable, Category = "Game|Components|ChaosVehicleMovement")
+	void SetYawInput(float Yaw);
+
+	/** Set the user input for handbrake */
+	UFUNCTION(BlueprintCallable, Category = "Game|Components|ChaosVehicleMovement")
+	void SetHandbrakeInput(float Handbrake);
+
+	/** Set the gear directly */
+	UFUNCTION(BlueprintCallable, Category = "Game|Components|ChaosVehicleMovement")
+	void SetGearInput(int Gear);
+
+	void SetSleeping(bool bEnableSleep);
+
+	//~ Begin Networking Replication Interface.
+	UPROPERTY(Replicated)
+	FModularVehicleRepData ModularVehicleRepData;
+
+	int32 MVRepVersionProcessed = INDEX_NONE;
+
+	/** Called post solve to allow authoritative components to update their replication data */
+	void UpdateMVRepData();
+
+	/** Called post solve to allow clients to receive their replication data */
+	void ProcessMVRepData();
+
+	virtual bool ProcessRepData(float DeltaTime, float SimTime) override;
+	//~ End Networking Replication Interface.
+
+	void PerformRaycasts(FModularVehicleDefaultAsyncInput* Inputs);
+	Chaos::FWheelSimModule* LocatePhysicallyClosestWheel(TUniquePtr<Chaos::FSimModuleTree>& SimModuleTree, int SuspensionTreeIndex);
+	void FixupTreeLinks(TUniquePtr<Chaos::FSimModuleTree>& SimModuleTree);
+	void EnableAnimationForPhysicsDrivenTransforms(TUniquePtr<Chaos::FSimModuleTree>& SimModuleTree);
+	void ClearRawInput();
+
+	void CalcThrottleBrakeInput(float& ThrottleOut, float& BrakeOut);
+	void UpdateChassisUnionCollision();
+
+	EChaosAsyncVehicleDataType CurAsyncType;
+	FModularVehicleAsyncInput* CurAsyncInput;
+	struct FModularVehicleAsyncOutput* CurAsyncOutput;
+	struct FModularVehicleAsyncOutput* NextAsyncOutput;
+	float OutputInterpAlpha;
+
+	struct FAsyncOutputWrapper
+	{
+		int32 Idx;
+		int32 Timestamp;
+
+		FAsyncOutputWrapper()
+			: Idx(INDEX_NONE)
+			, Timestamp(INDEX_NONE)
+		{
+		}
+	};
+	TArray<FAsyncOutputWrapper> OutputsWaitingOn;
+	TUniquePtr<FPhysicsVehicleOutput> PVehicleOutput;	/* physics simulation data output from the async physics thread */
+	TUniquePtr<FModularVehicleSimulationGC> VehicleSimulationPT;	/* simulation code running on the physics thread async callback */
+
+	// Method of applying control inputs at a specific time - for network prediction
+	UFUNCTION(reliable, server)
+	void ServerApplyControls(float InSteeringInput, float InThrottleInput, float InBrakeInput
+			, float InHandbrakeInput, int32 InGearInput, float InRollInput, float InPitchInput, float InYawInput, const FAsyncPhysicsTimestamp& AsyncPhysicsTimestamp);
+
+private:
+	// The function that actually applies controls - designed to run on both server and client
+	void ApplyControls_Imp(float InSteeringInput, float InThrottleInput, float InBrakeInput
+		, float InHandbrakeInput, int32 InGearInput, float InRollInput, float InPitchInput, float InYawInput, const FAsyncPhysicsTimestamp& AsyncPhysicsTimestamp);
+
+	Chaos::FPBDRigidClusteredParticleHandle* GetChassisParticle();
+
+	bool bUsingNetworkPhysicsPrediction;
+	float PrevSteeringInput;
+};
+
+
diff --git a/Engine/Plugins/Experimental/ChaosModularVehicle/Source/ChaosModularVehicleEngine/Public/ChaosModularVehicle/ModularVehicleInputRate.h b/Engine/Plugins/Experimental/ChaosModularVehicle/Source/ChaosModularVehicleEngine/Public/ChaosModularVehicle/ModularVehicleInputRate.h
new file mode 100644
--- /dev/null
+++ b/Engine/Plugins/Experimental/ChaosModularVehicle/Source/ChaosModularVehicleEngine/Public/ChaosModularVehicle/ModularVehicleInputRate.h
@@ -0,0 +1,89 @@
+// Copyright Epic Games, Inc. All Rights Reserved.
+
+#pragma once
+#include "SimModule/SimModulesInclude.h"
+#include "Curves/CurveFloat.h"
+
+#include "ModularVehicleInputRate.generated.h"
+
+struct CHAOSMODULARVEHICLEENGINE_API FModularVehicleDebugParams
+{
+	bool ShowDebug = false;
+	bool SuspensionRaycastsEnabled = true;
+	bool ShowSuspensionRaycasts = false;
+	bool ShowWheelData = false;
+	bool ShowRaycastMaterial = false;
+	bool ShowWheelCollisionNormal = false;
+
+	bool DisableAnim = false;
+	float FrictionOverride = 1.0f;
+};
+
+enum EModularVehicleInputType : uint8
+{
+	Throttle = 0,
+	Brake,
+	Clutch,
+	Steering,
+	Handbrake,
+	Pitch,
+	Roll,
+	Yaw,
+	Gear,
+	DebugIndex,
+
+	Max
+};
+
+UENUM()
+enum class EModularVehicleInputFunctionType : uint8
+{
+	LinearFunction = 0,
+	SquaredFunction,
+	CustomCurve
+};
+
+
+struct CHAOSMODULARVEHICLEENGINE_API FModularVehicleHistory
+{
+	FModularVehicleHistory()
+		: ControlInputs()
+	{
+	}
+
+	// Control Inputs
+	Chaos::FControlInputs ControlInputs;
+};
+
+
+USTRUCT()
+struct CHAOSMODULARVEHICLEENGINE_API FModularVehicleInputRate
+{
+	GENERATED_USTRUCT_BODY()
+
+	FModularVehicleInputRate() : Name(TEXT("Empty")), RiseRate(5.0f), FallRate(5.0f), InputCurveFunction(EModularVehicleInputFunctionType::LinearFunction) {}
+
+	FModularVehicleInputRate(const FString& InName) : Name(InName), RiseRate(5.0f), FallRate(5.0f), InputCurveFunction(EModularVehicleInputFunctionType::LinearFunction) {}
+
+	~FModularVehicleInputRate() {}
+
+	UPROPERTY(VisibleAnywhere, Category = VehicleInputRate)
+	FString Name;
+
+	UPROPERTY(EditAnywhere, Category = VehicleInputRate)
+	float RiseRate;
+
+	UPROPERTY(EditAnywhere, Category = VehicleInputRate)
+	float FallRate;
+
+	UPROPERTY(EditAnywhere, Category = VehicleInputRate)
+	EModularVehicleInputFunctionType InputCurveFunction;
+
+	UPROPERTY(EditAnywhere, Category = VehicleInputRate)
+	FRuntimeFloatCurve UserCurve;
+
+	float InterpInputValue(float DeltaTime, float CurrentValue, float NewValue) const;
+
+	float CalcControlFunction(float InputValue);
+
+};
diff --git a/Engine/Plugins/Experimental/ChaosModularVehicle/Source/ChaosModularVehicleEngine/Public/ChaosModularVehicle/ModularVehicleObject.h b/Engine/Plugins/Experimental/ChaosModularVehicle/Source/ChaosModularVehicleEngine/Public/ChaosModularVehicle/ModularVehicleObject.h
new file mode 100644
--- /dev/null
+++ b/Engine/Plugins/Experimental/ChaosModularVehicle/Source/ChaosModularVehicleEngine/Public/ChaosModularVehicle/ModularVehicleObject.h
@@ -0,0 +1,25 @@
+// Copyright Epic Games, Inc. All Rights Reserved.
+
+#pragma once
+
+#include "GeometryCollection/GeometryCollectionObject.h"
+#include "ChaosModularVehicle/ModularSimCollection.h"
+
+#include "ModularVehicleObject.generated.h"
+
+namespace Chaos
+{
+	class FSimModuleTree;
+	class ISimulationModuleBase;
+}
+
+UCLASS(BlueprintType, customconstructor)
+class CHAOSMODULARVEHICLEENGINE_API UModularVehicle : public UGeometryCollection
+{
+	GENERATED_UCLASS_BODY()
+	public:
+
+	UModularVehicle(const FObjectInitializer& ObjectInitializer = FObjectInitializer::Get());
+
+};
+
diff --git a/Engine/Plugins/Experimental/ChaosModularVehicle/Source/ChaosModularVehicleEngine/Public/ChaosModularVehicle/ModularVehiclePawn.h b/Engine/Plugins/Experimental/ChaosModularVehicle/Source/ChaosModularVehicleEngine/Public/ChaosModularVehicle/ModularVehiclePawn.h
new file mode 100644
--- /dev/null
+++ b/Engine/Plugins/Experimental/ChaosModularVehicle/Source/ChaosModularVehicleEngine/Public/ChaosModularVehicle/ModularVehiclePawn.h
@@ -0,0 +1,35 @@
+// Copyright Epic Games, Inc. All Rights Reserved.
+
+#pragma once
+
+#include "CoreMinimal.h"
+#include "GameFramework/Pawn.h"
+#include "UObject/ObjectMacros.h"
+#include "UObject/ObjectPtr.h"
+#include "ChaosModularVehicle/ModularVehicleComponent.h"
+
+#include "ModularVehiclePawn.generated.h"
+
+UCLASS()
+class CHAOSMODULARVEHICLEENGINE_API AModularVehiclePawn : public APawn
+{
+	GENERATED_UCLASS_BODY()
+
+protected:
+	virtual void BeginPlay() override;
+
+public:
+
+	/* Game state callback */
+	void TickSimulation(float DeltaTime);
+	virtual void Tick(float DeltaSeconds) override;
+
+	virtual void PostInitializeComponents() override;
+
+	/* ModularVehicleComponent */
+	UPROPERTY(VisibleAnywhere, BlueprintReadOnly, Category = Destruction, meta = (ExposeFunctionCategories = "Components|ModularVehicle", AllowPrivateAccess = "true"))
+	TObjectPtr<UModularVehicleComponent> ModularVehicleComponent;
+	
+	UModularVehicleComponent* GetModularVehicleComponent() const { return ModularVehicleComponent; }
+};
+
diff --git a/Engine/Plugins/Experimental/ChaosModularVehicle/Source/ChaosModularVehicleEngine/Public/ChaosModularVehicle/ModularVehicleSimulationCU.h b/Engine/Plugins/Experimental/ChaosModularVehicle/Source/ChaosModularVehicleEngine/Public/ChaosModularVehicle/ModularVehicleSimulationCU.h
new file mode 100644
--- /dev/null
+++ b/Engine/Plugins/Experimental/ChaosModularVehicle/Source/ChaosModularVehicleEngine/Public/ChaosModularVehicle/ModularVehicleSimulationCU.h
@@ -0,0 +1,70 @@
+// Copyright Epic Games, Inc. All Rights Reserved.
+
+#pragma once
+
+#include "GeometryCollection/GeometryCollectionComponent.h"
+#include "ChaosModularVehicle/ModularVehicleInputRate.h"
+#include "ChaosModularVehicle/ChaosSimModuleManagerAsyncCallback.h"
+#include "SimModule/SimModuleTree.h"
+#include "SimModule/SimModulesInclude.h"
+#include "ModularVehicleBuilder.h"
+
+struct FModularVehicleAsyncInput;
+struct FChaosSimModuleManagerAsyncOutput;
+
+namespace Chaos
+{
+	class FClusterUnionPhysicsProxy;
+}
+
+
+class CHAOSMODULARVEHICLEENGINE_API FModularVehicleSimulationCU
+{
+public:
+	FModularVehicleSimulationCU(bool InUsingNetworkPhysicsPrediction)
+		: bUsingNetworkPhysicsPrediction(InUsingNetworkPhysicsPrediction)
+	{
+	}
+
+	virtual ~FModularVehicleSimulationCU()
+	{
+		SimModuleTree.Reset();
+	}
+
+	void Initialize(TUniquePtr<Chaos::FSimModuleTree>& InSimModuleTree);
+	void Terminate();
+
+	/** Update called from Physics Thread */
+	virtual void Simulate(UWorld* InWorld, float DeltaSeconds, const FModularVehicleAsyncInput& InputData, FModularVehicleAsyncOutput& OutputData, IPhysicsProxyBase* Proxy);
+	virtual void Simulate_ClusterUnion(UWorld* InWorld, float DeltaSeconds, const FModularVehicleAsyncInput& InputData, FModularVehicleAsyncOutput& OutputData, Chaos::FClusterUnionPhysicsProxy* Proxy);
+
+	void ApplyDeferredForces(FGeometryCollectionPhysicsProxy* RigidHandle);
+	void ApplyDeferredForces(Chaos::FClusterUnionPhysicsProxy* Proxy);
+
+	void PerformAdditionalSimWork(UWorld* InWorld, const FModularVehicleAsyncInput& InputData, Chaos::FClusterUnionPhysicsProxy* Proxy, Chaos::FAllInputs& AllInputs);
+
+	void FillOutputState(FModularVehicleAsyncOutput& Output);
+
+	Chaos::FControlInputs& AccessControlInputs();
+
+	const TUniquePtr<Chaos::FSimModuleTree>& GetSimComponentTree() const {
+		Chaos::EnsureIsInPhysicsThreadContext();
+		return SimModuleTree;
+		}
+
+	TUniquePtr<Chaos::FSimModuleTree>& AccessSimComponentTree() {
+		return SimModuleTree; 
+		}
+
+	TArray<FModularVehicleInputRate>& AccessInputInterpolation() { return InputInterpolation; }
+
+
+	TUniquePtr<Chaos::FSimModuleTree> SimModuleTree;	/* Simulation modules stored in tree structure */
+	TArray<FModularVehicleInputRate> InputInterpolation;
+	Chaos::FAllInputs SimInputData;
+	bool bUsingNetworkPhysicsPrediction;
+
+	/** Current control inputs that is being used on the PT */
+	FModularVehicleInputs VehicleInputs;
+
+};
\ No newline at end of file
diff --git a/Engine/Plugins/Experimental/ChaosModularVehicle/Source/ChaosModularVehicleEngine/Public/ChaosModularVehicle/ModularVehicleSimulationGC.h b/Engine/Plugins/Experimental/ChaosModularVehicle/Source/ChaosModularVehicleEngine/Public/ChaosModularVehicle/ModularVehicleSimulationGC.h
new file mode 100644
--- /dev/null
+++ b/Engine/Plugins/Experimental/ChaosModularVehicle/Source/ChaosModularVehicleEngine/Public/ChaosModularVehicle/ModularVehicleSimulationGC.h
@@ -0,0 +1,78 @@
+// Copyright Epic Games, Inc. All Rights Reserved.
+
+#pragma once
+
+#include "GeometryCollection/GeometryCollectionComponent.h"
+#include "ChaosModularVehicle/ModularVehicleInputRate.h"
+#include "ChaosModularVehicle/ChaosSimModuleManagerAsyncCallback.h"
+#include "SimModule/SimModuleTree.h"
+#include "SimModule/SimModulesInclude.h"
+#include "Curves/CurveFloat.h"
+
+struct FModularVehicleAsyncInput;
+struct FChaosSimModuleManagerAsyncOutput;
+struct FModularVehicleDefaultAsyncInput;
+
+
+
+class CHAOSMODULARVEHICLEENGINE_API FModularVehicleSimulationGC
+{
+public:
+	FModularVehicleSimulationGC(bool InUsingNetworkPhysicsPrediction)
+		: bUsingNetworkPhysicsPrediction(InUsingNetworkPhysicsPrediction)
+		, StartFrame(0)
+	{
+	}
+
+	virtual ~FModularVehicleSimulationGC()
+	{
+		SimModuleTree.Reset();
+	}
+
+	void Initialize(TUniquePtr<Chaos::FSimModuleTree>& InSimModuleTree);
+	void Terminate();
+
+	void Tick(UWorld* InWorld, float DeltaSeconds, const FModularVehicleDefaultAsyncInput& InputData, FModularVehicleAsyncOutput& OutputData, FGeometryCollectionPhysicsProxy* Proxy)
+	{
+		Simulate(InWorld, DeltaSeconds, InputData, OutputData, Proxy);
+		FillOutputState(OutputData);
+	}
+
+	void SyncHistoryInputs(float DeltaSeconds, FGeometryCollectionPhysicsProxy* Proxy);
+	void InterpolateInputs(float DeltaSeconds, const Chaos::FControlInputs& ExternalInputIn, Chaos::FControlInputs& InterpolatedInputsInOut);
+
+	/** Update called from Physics Thread */
+	virtual void Simulate(UWorld* InWorld, float DeltaSeconds, const FModularVehicleDefaultAsyncInput& InputData, FModularVehicleAsyncOutput& OutputData, FGeometryCollectionPhysicsProxy* Proxy);
+
+	void ApplyDeferredForces(FGeometryCollectionPhysicsProxy* RigidHandle);
+
+	void PerformAdditionalSimWork(UWorld* InWorld, const FModularVehicleDefaultAsyncInput& InputData, FGeometryCollectionPhysicsProxy* Proxy, Chaos::FAllInputs& AllInputs);
+
+	void FillOutputState(FModularVehicleAsyncOutput& Output);
+
+	Chaos::FControlInputs& AccessControlInputs();
+
+	const TUniquePtr<Chaos::FSimModuleTree>& GetSimComponentTree() const {
+		Chaos::EnsureIsInPhysicsThreadContext();
+		return SimModuleTree;
+		}
+
+	TUniquePtr<Chaos::FSimModuleTree>& AccessSimComponentTree() {
+		Chaos::EnsureIsInPhysicsThreadContext();
+		return SimModuleTree; 
+		}
+
+	TArray<FModularVehicleInputRate>& AccessInputInterpolation() { return InputInterpolation; }
+
+private:
+
+	TUniquePtr<Chaos::FSimModuleTree> SimModuleTree;	/* Simulation modules stored in tree structure */
+	TArray<FModularVehicleInputRate> InputInterpolation;
+	Chaos::FAllInputs SimInputData;
+	bool bUsingNetworkPhysicsPrediction;
+
+	Chaos::FControlInputs AsyncControlInputs;  // Predictive networking relies on inputs being actioned ASAP via ExecuteAsyncPhysicsCommand with timestamp info
+	TArray<FModularVehicleHistory> VehicleHistoryBuffer;
+	int32 StartFrame = 0;
+	int32 NumHistoryFrames = 0;
+};
\ No newline at end of file
diff --git a/Engine/Plugins/Experimental/ChaosModularVehicle/Source/ChaosModularVehicleEngine/Public/ChaosModularVehicle/VehicleSimClutchComponent.h b/Engine/Plugins/Experimental/ChaosModularVehicle/Source/ChaosModularVehicleEngine/Public/ChaosModularVehicle/VehicleSimClutchComponent.h
new file mode 100644
--- /dev/null
+++ b/Engine/Plugins/Experimental/ChaosModularVehicle/Source/ChaosModularVehicleEngine/Public/ChaosModularVehicle/VehicleSimClutchComponent.h
@@ -0,0 +1,27 @@
+// Copyright Epic Games, Inc. All Rights Reserved.
+
+#pragma once
+
+#include "CoreMinimal.h"
+#include "UObject/ObjectPtr.h"
+#include "VehicleSimBaseComponent.h"
+#include "VehicleSimClutchComponent.generated.h"
+
+
+UCLASS(ClassGroup = (ModularVehicle), meta = (BlueprintSpawnableComponent), hidecategories = (Object, Tick, Replication, Cooking, Activation, LOD))
+class CHAOSMODULARVEHICLEENGINE_API UVehicleSimClutchComponent : public UVehicleSimBaseComponent
+{
+	GENERATED_BODY()
+
+public:
+	UVehicleSimClutchComponent();
+	virtual ~UVehicleSimClutchComponent() = default;
+
+	UPROPERTY(EditAnywhere, BlueprintReadWrite, Category = Attributes)
+		float ClutchStrength;
+
+	virtual ESimModuleType GetModuleType() const override { return ESimModuleType::Clutch; }
+
+	virtual Chaos::ISimulationModuleBase* CreateNewCoreModule() const override;
+
+};
\ No newline at end of file
