[Deprecation Cleanup] Remove deprecated code in actor factory class

diff --git a/Engine/Plugins/Experimental/ChaosModularVehicle/Source/ChaosModularVehicleEditor/Private/ChaosModularVehicle/ActorFactoryModularVehicle.cpp b/Engine/Plugins/Experimental/ChaosModularVehicle/Source/ChaosModularVehicleEditor/Private/ChaosModularVehicle/ActorFactoryModularVehicle.cpp
--- a/Engine/Plugins/Experimental/ChaosModularVehicle/Source/ChaosModularVehicleEditor/Private/ChaosModularVehicle/ActorFactoryModularVehicle.cpp
+++ b/Engine/Plugins/Experimental/ChaosModularVehicle/Source/ChaosModularVehicleEditor/Private/ChaosModularVehicle/ActorFactoryModularVehicle.cpp
@@ -50,13 +50,4 @@ void UActorFactoryModularVehicle::PostSpawnActor(UObject* Asset, AActor* NewActo
 	NewModularVehicleActor->GetModularVehicleComponent()->RegisterComponent();
 }
 
-void UActorFactoryModularVehicle::PostCreateBlueprint(UObject* Asset, AActor* CDO)
-{
-	if (Asset != NULL && CDO != NULL)
-	{
-		UModularVehicleAsset* ModularVehicle = CastChecked<UModularVehicleAsset>(Asset);
-		AModularVehiclePawn* ModularVehicleActor = CastChecked<AModularVehiclePawn>(CDO);
-	}
-}
-
 #undef LOCTEXT_NAMESPACE
\ No newline at end of file
diff --git a/Engine/Plugins/Experimental/ChaosModularVehicle/Source/ChaosModularVehicleEditor/Private/ChaosModularVehicle/ActorFactoryModularVehicle.h b/Engine/Plugins/Experimental/ChaosModularVehicle/Source/ChaosModularVehicleEditor/Private/ChaosModularVehicle/ActorFactoryModularVehicle.h
--- a/Engine/Plugins/Experimental/ChaosModularVehicle/Source/ChaosModularVehicleEditor/Private/ChaosModularVehicle/ActorFactoryModularVehicle.h
+++ b/Engine/Plugins/Experimental/ChaosModularVehicle/Source/ChaosModularVehicleEditor/Private/ChaosModularVehicle/ActorFactoryModularVehicle.h
@@ -18,7 +18,6 @@ class UActorFactoryModularVehicle : public UActorFactory
 
 	//~ Begin UActorFactory Interface
 	virtual void PostSpawnActor( UObject* Asset, AActor* NewActor ) override;
-	virtual void PostCreateBlueprint( UObject* Asset, AActor* CDO ) override;
 	virtual bool CanCreateActorFrom( const FAssetData& AssetData, FText& OutErrorMsg ) override;
 	//~ End UActorFactory Interface
 };
