Modular Vehicles: More debugging options

diff --git a/Engine/Plugins/Experimental/ChaosModularVehicle/Source/ChaosModularVehicleEngine/Private/ChaosModularVehicle/ChaosSimModuleManager.cpp b/Engine/Plugins/Experimental/ChaosModularVehicle/Source/ChaosModularVehicleEngine/Private/ChaosModularVehicle/ChaosSimModuleManager.cpp
--- a/Engine/Plugins/Experimental/ChaosModularVehicle/Source/ChaosModularVehicleEngine/Private/ChaosModularVehicle/ChaosSimModuleManager.cpp
+++ b/Engine/Plugins/Experimental/ChaosModularVehicle/Source/ChaosModularVehicleEngine/Private/ChaosModularVehicle/ChaosSimModuleManager.cpp
@@ -80,8 +80,8 @@ void FChaosSimModuleManager::OnWorldCleanup(UWorld* InWorld, bool bSessionEnded,
 
 void FChaosSimModuleManager::OnShowDebugInfo(AHUD* HUD, UCanvas* Canvas, const FDebugDisplayInfo& DisplayInfo, float& YL, float& YPos)
 {
-	static const FName NAME_VehicleModuleManager("VehicleModuleManager");
-	if (Canvas && HUD->ShouldDisplayDebug(NAME_VehicleModuleManager))
+	static const FName NAME_ModularVehicle("ModularVehicle");
+	if (Canvas && HUD->ShouldDisplayDebug(NAME_ModularVehicle))
 	{
 		FChaosSimModuleManager* Manager = FChaosSimModuleManager::GetManagerFromScene(HUD->GetWorld()->GetPhysicsScene());
 
diff --git a/Engine/Plugins/Experimental/ChaosModularVehicle/Source/ChaosModularVehicleEngine/Private/ChaosModularVehicle/ModularVehicleBaseComponent.cpp b/Engine/Plugins/Experimental/ChaosModularVehicle/Source/ChaosModularVehicleEngine/Private/ChaosModularVehicle/ModularVehicleBaseComponent.cpp
--- a/Engine/Plugins/Experimental/ChaosModularVehicle/Source/ChaosModularVehicleEngine/Private/ChaosModularVehicle/ModularVehicleBaseComponent.cpp
+++ b/Engine/Plugins/Experimental/ChaosModularVehicle/Source/ChaosModularVehicleEngine/Private/ChaosModularVehicle/ModularVehicleBaseComponent.cpp
@@ -688,6 +688,9 @@ void UModularVehicleBaseComponent::ParallelUpdate(float DeltaTime)
 							EngineTorque = Engine->Torque;
 						}
 
+#if !(UE_BUILD_SHIPPING || UE_BUILD_TEST)
+						PVehicleOutput->SimTreeOutputData[I]->DebugString = NextOutput->VehicleSimOutput.SimTreeOutputData[I]->DebugString;
+#endif
 					}
 				}
 			}
@@ -800,7 +803,7 @@ void UModularVehicleBaseComponent::AddComponentToSimulation(UPrimitiveComponent*
 	int32 ComponentAddOrder = INDEX_NONE;
 	if (InComponent && bIsNew)
 	{
-		UE_LOG(LogModularBase, Log, TEXT("AddComponentToSimulation %s"), *InComponent->GetName());
+		UE_LOG(LogModularBase, Log, TEXT("AddComponentToSimulation %s, NetMode %d"), *InComponent->GetName(), InComponent->GetNetMode());
 		
 		if (ClusterUnionComponent->IsAuthority())
 		{
@@ -842,7 +845,7 @@ void UModularVehicleBaseComponent::RemoveComponentFromSimulation(UPrimitiveCompo
 {
 	if (InComponent && VehicleSimulationPT)
 	{
-		UE_LOG(LogModularBase, Log, TEXT("UModularVehicleBaseComponent::RemoveComponentFromSimulation: %s"), *InComponent->GetFullName());
+		UE_LOG(LogModularBase, Log, TEXT("UModularVehicleBaseComponent::RemoveComponentFromSimulation: %s, NetMode %d"), *InComponent->GetFullName(), InComponent->GetNetMode());
 
 		Chaos::FSimTreeUpdates LatestTreeUpdates;
 
@@ -905,7 +908,7 @@ void UModularVehicleBaseComponent::CreateVehicleSim()
 	UWorld* World = GetWorld();
 
 	// create the simulation class
-	VehicleSimulationPT = MakeUnique<FModularVehicleSimulationCU>(bUsingNetworkPhysicsPrediction);
+	VehicleSimulationPT = MakeUnique<FModularVehicleSimulationCU>(bUsingNetworkPhysicsPrediction, (int8)World->GetNetMode());
 
 	// create physics output container
 	PVehicleOutput = MakeUnique<FPhysicsVehicleOutput>();
@@ -1109,17 +1112,24 @@ void UModularVehicleBaseComponent::ShowDebugInfo(AHUD* HUD, UCanvas* Canvas, con
 
 	// draw input values
 	Canvas->SetDrawColor(FColor::White);
-	YPos += Canvas->DrawText(RenderFont, FString::Printf(TEXT("Throttle Raw  (%f) %f"), RawThrottleInput, ThrottleInput), 4, YPos);
-	YPos += Canvas->DrawText(RenderFont, FString::Printf(TEXT("Boost Raw     (%f) %f"), RawBoostInput, BoostInput), 4, YPos);
-	YPos += Canvas->DrawText(RenderFont, FString::Printf(TEXT("Drift Raw     (%f) %f"), RawDriftInput, DriftInput), 4, YPos);
-	YPos += Canvas->DrawText(RenderFont, FString::Printf(TEXT("Brake Raw     (%f) %f"), RawBrakeInput, BrakeInput), 4, YPos);
-	YPos += Canvas->DrawText(RenderFont, FString::Printf(TEXT("Steering Raw  (%f) %f"), RawSteeringInput, SteeringInput), 4, YPos);
-	YPos += Canvas->DrawText(RenderFont, FString::Printf(TEXT("Handbrake Raw (%f) %f"), RawHandbrakeInput, HandbrakeInput), 4, YPos);
-	YPos += Canvas->DrawText(RenderFont, FString::Printf(TEXT("Roll Raw      (%f) %f"), RawRollInput, RollInput), 4, YPos);
-	YPos += Canvas->DrawText(RenderFont, FString::Printf(TEXT("Pitch Raw     (%f) %f"), RawPitchInput, PitchInput), 4, YPos);
-	YPos += Canvas->DrawText(RenderFont, FString::Printf(TEXT("Yaw Raw       (%f) %f"), RawYawInput, YawInput), 4, YPos);
-	YPos += Canvas->DrawText(RenderFont, FString::Printf(TEXT("Reverse Raw       (%f) %f"), RawReverseInput, ReverseInput), 4, YPos);
+	YPos += Canvas->DrawText(RenderFont, FString::Printf(TEXT("Throttle Raw  (%3.2f) %3.2f"), RawThrottleInput, ThrottleInput), 4, YPos);
+	YPos += Canvas->DrawText(RenderFont, FString::Printf(TEXT("Boost Raw     (%3.2f) %3.2f"), RawBoostInput, BoostInput), 4, YPos);
+	YPos += Canvas->DrawText(RenderFont, FString::Printf(TEXT("Drift Raw     (%3.2f) %3.2f"), RawDriftInput, DriftInput), 4, YPos);
+	YPos += Canvas->DrawText(RenderFont, FString::Printf(TEXT("Brake Raw     (%3.2f) %3.2f"), RawBrakeInput, BrakeInput), 4, YPos);
+	YPos += Canvas->DrawText(RenderFont, FString::Printf(TEXT("Steering Raw  (%3.2f) %3.2f"), RawSteeringInput, SteeringInput), 4, YPos);
+	YPos += Canvas->DrawText(RenderFont, FString::Printf(TEXT("Handbrake Raw (%3.2f) %3.2f"), RawHandbrakeInput, HandbrakeInput), 4, YPos);
+	YPos += Canvas->DrawText(RenderFont, FString::Printf(TEXT("Roll Raw      (%3.2f) %3.2f"), RawRollInput, RollInput), 4, YPos);
+	YPos += Canvas->DrawText(RenderFont, FString::Printf(TEXT("Pitch Raw     (%3.2f) %3.2f"), RawPitchInput, PitchInput), 4, YPos);
+	YPos += Canvas->DrawText(RenderFont, FString::Printf(TEXT("Yaw Raw       (%3.2f) %3.2f"), RawYawInput, YawInput), 4, YPos);
+	YPos += Canvas->DrawText(RenderFont, FString::Printf(TEXT("Reverse Raw   (%3.2f) %3.2f"), RawReverseInput, ReverseInput), 4, YPos);
 
 	YPos += 10;
 
+#if !(UE_BUILD_SHIPPING || UE_BUILD_TEST)
+	for (Chaos::FSimOutputData* Data : PVehicleOutput->SimTreeOutputData)
+	{
+		YPos += Canvas->DrawText(RenderFont, FString::Printf(TEXT("%s"), *Data->ToString()), 4, YPos);
+	}
+#endif
+
 }
diff --git a/Engine/Plugins/Experimental/ChaosModularVehicle/Source/ChaosModularVehicleEngine/Private/ChaosModularVehicle/ModularVehicleComponent.cpp b/Engine/Plugins/Experimental/ChaosModularVehicle/Source/ChaosModularVehicleEngine/Private/ChaosModularVehicle/ModularVehicleComponent.cpp
--- a/Engine/Plugins/Experimental/ChaosModularVehicle/Source/ChaosModularVehicleEngine/Private/ChaosModularVehicle/ModularVehicleComponent.cpp
+++ b/Engine/Plugins/Experimental/ChaosModularVehicle/Source/ChaosModularVehicleEngine/Private/ChaosModularVehicle/ModularVehicleComponent.cpp
@@ -994,8 +994,6 @@ void UModularVehicleComponent::ParallelUpdate(float DeltaSeconds)
 	{
 		if (CurrentOutput->bValid && PVehicleOutput)
 		{
-			PVehicleOutput->DebugStrings = MoveTemp(CurrentOutput->VehicleSimOutput.DebugStrings);
-
 			// #TODO: lerp output values here
 			// i.e. PVehicleOutput->Value = FMath::Lerp(CurAsyncOutput->VehicleSimOutput.Value, NextAsyncOutput->VehicleSimOutput.Value, OutputInterpAlpha);
 			//for (int I=0; I<CurAsyncOutput->VehicleSimOutput.ModuleTransform.Num(); I++)
@@ -1049,12 +1047,12 @@ void UModularVehicleComponent::ShowDebugInfo(AHUD* HUD, UCanvas* Canvas, const F
 		Canvas->SetDrawColor(FColor::White);
 		YPos += 16;
 
-		for (const FString& StringOut : PVehicleOutput->DebugStrings)
-		{
-			YPos += Canvas->DrawText(RenderFont, FString::Printf(TEXT("%s"), *StringOut), 4, YPos);
-		}
+		//for (const FString& StringOut : PVehicleOutput->DebugStrings)
+		//{
+		//	YPos += Canvas->DrawText(RenderFont, FString::Printf(TEXT("%s"), *StringOut), 4, YPos);
+		//}
 
-		PVehicleOutput->DebugStrings.Empty();
+		//PVehicleOutput->DebugStrings.Empty();
 	}
 
 	YPos += 10;
diff --git a/Engine/Plugins/Experimental/ChaosModularVehicle/Source/ChaosModularVehicleEngine/Private/ChaosModularVehicle/ModularVehicleSimulationCU.cpp b/Engine/Plugins/Experimental/ChaosModularVehicle/Source/ChaosModularVehicleEngine/Private/ChaosModularVehicle/ModularVehicleSimulationCU.cpp
--- a/Engine/Plugins/Experimental/ChaosModularVehicle/Source/ChaosModularVehicleEngine/Private/ChaosModularVehicle/ModularVehicleSimulationCU.cpp
+++ b/Engine/Plugins/Experimental/ChaosModularVehicle/Source/ChaosModularVehicleEngine/Private/ChaosModularVehicle/ModularVehicleSimulationCU.cpp
@@ -15,7 +15,6 @@
 FModularVehicleDebugParams GModularVehicleDebugParams;
 
 #if CHAOS_DEBUG_DRAW
-FAutoConsoleVariableRef CVarChaosModularVehiclesShowDebug(TEXT("p.ModularVehicle.ShowDebug"), GModularVehicleDebugParams.ShowDebug, TEXT("Enable/Disable Show Modular Vehicle Debug on HUD."));
 FAutoConsoleVariableRef CVarChaosModularVehiclesRaycastsEnabled(TEXT("p.ModularVehicle.SuspensionRaycastsEnabled"), GModularVehicleDebugParams.SuspensionRaycastsEnabled, TEXT("Enable/Disable Suspension Raycasts."));
 FAutoConsoleVariableRef CVarChaosModularVehiclesShowRaycasts(TEXT("p.ModularVehicle.ShowSuspensionRaycasts"), GModularVehicleDebugParams.ShowSuspensionRaycasts, TEXT("Enable/Disable Suspension Raycast Visualisation."));
 FAutoConsoleVariableRef CVarChaosModularVehiclesShowWheelData(TEXT("p.ModularVehicle.ShowWheelData"), GModularVehicleDebugParams.ShowWheelData, TEXT("Enable/Disable Displaying Wheel Simulation Data."));
@@ -267,7 +266,19 @@ void FModularVehicleSimulationCU::PerformAdditionalSimWork(UWorld* InWorld, cons
 										if (GModularVehicleDebugParams.ShowWheelData)
 										{
 											FString TextOut = FString::Format(TEXT("{0}"), { Wheel->GetForceIntoSurface() });
-											Chaos::FDebugDrawQueue::GetInstance().DrawDebugString(HitResult.ImpactPoint + FVec3(0, 50, 50), TextOut, nullptr, FColor::White, -1.f, true, 1.0f);
+											FColor Col = FColor::White;
+											if (InWorld)
+											{
+												if (InWorld->GetNetMode() == ENetMode::NM_Client)
+												{
+													Col = FColor::Blue;
+												}
+												else
+												{
+													Col = FColor::Red;
+												}
+											}
+											Chaos::FDebugDrawQueue::GetInstance().DrawDebugString(HitResult.ImpactPoint + FVec3(0, 50, 50), TextOut, nullptr, Col, -1.f, true, 1.0f);
 										}
 									}
 								}
@@ -284,7 +295,20 @@ void FModularVehicleSimulationCU::PerformAdditionalSimWork(UWorld* InWorld, cons
 								Chaos::FDebugDrawQueue::GetInstance().DrawDebugSphere(TraceStart, 3, 16, FColor::White, false, -1.f, 0, 10.f);
 								Chaos::FDebugDrawQueue::GetInstance().DrawDebugSphere(HitResult.ImpactPoint, 1, 16, FColor::Red, false, -1.f, 0, 10.f);
 								FString TextOut = FString::Format(TEXT("{0}"), { HitResult.Time});
-								Chaos::FDebugDrawQueue::GetInstance().DrawDebugString(HitResult.ImpactPoint + FVec3(0, 50, 50), TextOut, nullptr, FColor::White, -1.f, true, 1.0f);
+
+								FColor Col = FColor::White;
+								if (InWorld)
+								{
+									if (InWorld->GetNetMode() == ENetMode::NM_Client)
+									{
+										Col = FColor::Blue;
+									}
+									else
+									{
+										Col = FColor::Red;
+									}
+								}
+								Chaos::FDebugDrawQueue::GetInstance().DrawDebugString(HitResult.ImpactPoint + FVec3(0, 50, 50), TextOut, nullptr, Col, -1.f, true, 1.0f);
 							}
 
 							if (GModularVehicleDebugParams.ShowRaycastMaterial)
diff --git a/Engine/Plugins/Experimental/ChaosModularVehicle/Source/ChaosModularVehicleEngine/Public/ChaosModularVehicle/ChaosSimModuleManagerAsyncCallback.h b/Engine/Plugins/Experimental/ChaosModularVehicle/Source/ChaosModularVehicleEngine/Public/ChaosModularVehicle/ChaosSimModuleManagerAsyncCallback.h
--- a/Engine/Plugins/Experimental/ChaosModularVehicle/Source/ChaosModularVehicleEngine/Public/ChaosModularVehicle/ChaosSimModuleManagerAsyncCallback.h
+++ b/Engine/Plugins/Experimental/ChaosModularVehicle/Source/ChaosModularVehicleEngine/Public/ChaosModularVehicle/ChaosSimModuleManagerAsyncCallback.h
@@ -192,7 +192,6 @@ struct CHAOSMODULARVEHICLEENGINE_API FPhysicsVehicleOutput
 	}
 
 	TArray<Chaos::FSimOutputData*> SimTreeOutputData;
-	TArray<FString> DebugStrings;
 };
 
 struct CHAOSMODULARVEHICLEENGINE_API FPhysicsModularVehicleTraits
diff --git a/Engine/Plugins/Experimental/ChaosModularVehicle/Source/ChaosModularVehicleEngine/Public/ChaosModularVehicle/ModularVehicleSimulationCU.h b/Engine/Plugins/Experimental/ChaosModularVehicle/Source/ChaosModularVehicleEngine/Public/ChaosModularVehicle/ModularVehicleSimulationCU.h
--- a/Engine/Plugins/Experimental/ChaosModularVehicle/Source/ChaosModularVehicleEngine/Public/ChaosModularVehicle/ModularVehicleSimulationCU.h
+++ b/Engine/Plugins/Experimental/ChaosModularVehicle/Source/ChaosModularVehicleEngine/Public/ChaosModularVehicle/ModularVehicleSimulationCU.h
@@ -21,8 +21,9 @@ namespace Chaos
 class CHAOSMODULARVEHICLEENGINE_API FModularVehicleSimulationCU
 {
 public:
-	FModularVehicleSimulationCU(bool InUsingNetworkPhysicsPrediction)
+	FModularVehicleSimulationCU(bool InUsingNetworkPhysicsPrediction, int8 InNetMode)
 		: bUsingNetworkPhysicsPrediction(InUsingNetworkPhysicsPrediction)
+		, NetMode(InNetMode)
 	{
 	}
 
@@ -67,4 +68,6 @@ public:
 	/** Current control inputs that is being used on the PT */
 	FModularVehicleInputs VehicleInputs;
 
+	int8 NetMode;
+
 };
\ No newline at end of file
diff --git a/Engine/Source/Runtime/Experimental/ChaosVehicles/ChaosVehiclesCore/Private/SimModule/EngineModule.cpp b/Engine/Source/Runtime/Experimental/ChaosVehicles/ChaosVehiclesCore/Private/SimModule/EngineModule.cpp
--- a/Engine/Source/Runtime/Experimental/ChaosVehicles/ChaosVehiclesCore/Private/SimModule/EngineModule.cpp
+++ b/Engine/Source/Runtime/Experimental/ChaosVehicles/ChaosVehiclesCore/Private/SimModule/EngineModule.cpp
@@ -74,6 +74,9 @@ namespace Chaos
 	void FEngineOutputData::FillOutputState(const ISimulationModuleBase* SimModule)
 	{
 		check(SimModule->GetSimType() == eSimType::Engine);
+
+		FSimOutputData::FillOutputState(SimModule);
+
 		if (const FEngineSimModule* Sim = static_cast<const FEngineSimModule*>(SimModule))
 		{
 			RPM = Sim->GetRPM();
@@ -89,11 +92,13 @@ namespace Chaos
 		Torque = FMath::Lerp(Current.Torque, Next.Torque, Alpha);
 	}
 
+#if !(UE_BUILD_SHIPPING || UE_BUILD_TEST)
 	FString FEngineOutputData::ToString()
 	{
-		return  FString::Printf(TEXT("RPM=%3.3f, Torque=%3.3f")
-			, RPM, Torque);
+		return FString::Printf(TEXT("%s, RPM=%3.3f, Torque=%3.3f")
+			, *DebugString, RPM, Torque);
 	}
+#endif
 
 } // namespace Chaos
 
diff --git a/Engine/Source/Runtime/Experimental/ChaosVehicles/ChaosVehiclesCore/Private/SimModule/SimulationModuleBase.cpp b/Engine/Source/Runtime/Experimental/ChaosVehicles/ChaosVehiclesCore/Private/SimModule/SimulationModuleBase.cpp
--- a/Engine/Source/Runtime/Experimental/ChaosVehicles/ChaosVehiclesCore/Private/SimModule/SimulationModuleBase.cpp
+++ b/Engine/Source/Runtime/Experimental/ChaosVehicles/ChaosVehiclesCore/Private/SimModule/SimulationModuleBase.cpp
@@ -114,4 +114,16 @@ const FTransform& ISimulationModuleBase::GetParentRelativeTransform() const
 }
 
 
+void FSimOutputData::FillOutputState(const ISimulationModuleBase* SimModule)
+{
+#if !(UE_BUILD_SHIPPING || UE_BUILD_TEST)
+	if (SimModule)
+	{	
+		DebugString.Empty();
+		SimModule->GetDebugString(DebugString);
+	}
+#endif
+}
+
+
 } //namespace Chaos
\ No newline at end of file
diff --git a/Engine/Source/Runtime/Experimental/ChaosVehicles/ChaosVehiclesCore/Private/SimModule/SuspensionModule.cpp b/Engine/Source/Runtime/Experimental/ChaosVehicles/ChaosVehiclesCore/Private/SimModule/SuspensionModule.cpp
--- a/Engine/Source/Runtime/Experimental/ChaosVehicles/ChaosVehiclesCore/Private/SimModule/SuspensionModule.cpp
+++ b/Engine/Source/Runtime/Experimental/ChaosVehicles/ChaosVehiclesCore/Private/SimModule/SuspensionModule.cpp
@@ -172,6 +172,9 @@ namespace Chaos
 	void FSuspensionOutputData::FillOutputState(const ISimulationModuleBase* SimModule)
 	{
 		check(SimModule->GetSimType() == eSimType::Suspension);
+
+		FSimOutputData::FillOutputState(SimModule);
+
 		if (const FSuspensionSimModule* Sim = static_cast<const FSuspensionSimModule*>(SimModule))
 		{
 			SpringDisplacement = Sim->SpringDisplacement;
@@ -188,10 +191,12 @@ namespace Chaos
 		SpringSpeed = FMath::Lerp(Current.SpringSpeed, Next.SpringSpeed, Alpha);
 	}
 
+#if !(UE_BUILD_SHIPPING || UE_BUILD_TEST)
 	FString FSuspensionOutputData::ToString()
 	{
-		return FString::Printf(TEXT("SpringDisplacement=%3.3f, SpringSpeed=%3.3f"), SpringDisplacement, SpringSpeed);
+		return FString::Printf(TEXT("%s, SpringDisplacement=%3.3f, SpringSpeed=%3.3f"), *DebugString, SpringDisplacement, SpringSpeed);
 	}
+#endif
 
 } // namespace Chaos
 
diff --git a/Engine/Source/Runtime/Experimental/ChaosVehicles/ChaosVehiclesCore/Private/SimModule/TransmissionModule.cpp b/Engine/Source/Runtime/Experimental/ChaosVehicles/ChaosVehiclesCore/Private/SimModule/TransmissionModule.cpp
--- a/Engine/Source/Runtime/Experimental/ChaosVehicles/ChaosVehiclesCore/Private/SimModule/TransmissionModule.cpp
+++ b/Engine/Source/Runtime/Experimental/ChaosVehicles/ChaosVehiclesCore/Private/SimModule/TransmissionModule.cpp
@@ -198,6 +198,9 @@ namespace Chaos
 	void FTransmissionOutputData::FillOutputState(const ISimulationModuleBase* SimModule)
 	{
 		check(SimModule->GetSimType() == eSimType::Transmission);
+
+		FSimOutputData::FillOutputState(SimModule);
+
 		if (const FTransmissionSimModule* Sim = static_cast<const FTransmissionSimModule*>(SimModule))
 		{
 			CurrentGear = Sim->CurrentGear;
@@ -212,11 +215,13 @@ namespace Chaos
 		CurrentGear = Current.CurrentGear;
 	}
 
+#if !(UE_BUILD_SHIPPING || UE_BUILD_TEST)
 	FString FTransmissionOutputData::ToString()
 	{
-		return  FString::Printf(TEXT("CurrentGear=%d")
-			, CurrentGear);
+		return FString::Printf(TEXT("%s CurrentGear=%d")
+			, *DebugString, CurrentGear);
 	}
+#endif
 
 } // namespace Chaos
 
diff --git a/Engine/Source/Runtime/Experimental/ChaosVehicles/ChaosVehiclesCore/Private/SimModule/WheelModule.cpp b/Engine/Source/Runtime/Experimental/ChaosVehicles/ChaosVehiclesCore/Private/SimModule/WheelModule.cpp
--- a/Engine/Source/Runtime/Experimental/ChaosVehicles/ChaosVehiclesCore/Private/SimModule/WheelModule.cpp
+++ b/Engine/Source/Runtime/Experimental/ChaosVehicles/ChaosVehiclesCore/Private/SimModule/WheelModule.cpp
@@ -229,6 +229,9 @@ namespace Chaos
 	inline void FWheelOutputData::FillOutputState(const ISimulationModuleBase* SimModule)
 	{
 		check(SimModule->GetSimType() == eSimType::Wheel);
+
+		FSimOutputData::FillOutputState(SimModule);
+
 		if (const FWheelSimModule* Sim = static_cast<const FWheelSimModule*>(SimModule))
 		{
 			bTouchingGround = Sim->bTouchingGround;
@@ -249,12 +252,13 @@ namespace Chaos
 		RPM = FMath::Lerp(Current.RPM, Next.RPM, Alpha);
 	}
 
+#if !(UE_BUILD_SHIPPING || UE_BUILD_TEST)
 	FString FWheelOutputData::ToString()
 	{
-		return  FString::Printf(TEXT("bTouchingGround=%d, ForceIntoSurface=%3.3f, SlipAngle=%3.3f, RPM=%3.3f")
-			, bTouchingGround, ForceIntoSurface, SlipAngle, RPM);
+		return  FString::Printf(TEXT("%s, bTouchingGround=%d, ForceIntoSurface=%3.3f, SlipAngle=%3.3f, RPM=%3.3f")
+			, *DebugString, bTouchingGround, ForceIntoSurface, SlipAngle, RPM);
 	}
-
+#endif
 
 } // namespace Chaos
 
diff --git a/Engine/Source/Runtime/Experimental/ChaosVehicles/ChaosVehiclesCore/Public/SimModule/EngineModule.h b/Engine/Source/Runtime/Experimental/ChaosVehicles/ChaosVehiclesCore/Public/SimModule/EngineModule.h
--- a/Engine/Source/Runtime/Experimental/ChaosVehicles/ChaosVehiclesCore/Public/SimModule/EngineModule.h
+++ b/Engine/Source/Runtime/Experimental/ChaosVehicles/ChaosVehiclesCore/Public/SimModule/EngineModule.h
@@ -42,7 +42,9 @@ namespace Chaos
 		virtual eSimType GetType() override { return eSimType::Engine; }
 		virtual void FillOutputState(const ISimulationModuleBase* SimModule) override;
 		virtual void Lerp(const FSimOutputData& InCurrent, const FSimOutputData& InNext, float Alpha) override;
+#if !(UE_BUILD_SHIPPING || UE_BUILD_TEST)
 		virtual FString ToString() override;
+#endif
 
 		float RPM;
 		float Torque;
diff --git a/Engine/Source/Runtime/Experimental/ChaosVehicles/ChaosVehiclesCore/Public/SimModule/SimulationModuleBase.h b/Engine/Source/Runtime/Experimental/ChaosVehicles/ChaosVehiclesCore/Public/SimModule/SimulationModuleBase.h
--- a/Engine/Source/Runtime/Experimental/ChaosVehicles/ChaosVehiclesCore/Public/SimModule/SimulationModuleBase.h
+++ b/Engine/Source/Runtime/Experimental/ChaosVehicles/ChaosVehiclesCore/Public/SimModule/SimulationModuleBase.h
@@ -148,7 +148,7 @@ namespace Chaos
 		Axle,			// connects more than one wheel
 		Transmission,	// gears - torque multiplier
 		Engine,			// (torque curve required) power source generates torque for wheel, axle, transmission, clutch
-		Motor,			// (electric?, no torque curve required?) power source generates torque for wheel, axle, transmission, clutch
+		Motor,			// NOT USED YET (electric?, no torque curve required?) power source generates torque for wheel, axle, transmission, clutch
 		Clutch,			// limits the amount of torque transferred between source and destination allowing for different rotation speeds of connected axles
 		Wing,			// lift and controls aircraft roll
 		Rudder,			// controls aircraft yaw
@@ -391,10 +391,17 @@ namespace Chaos
 		virtual ~FSimOutputData() {}
 
 		virtual eSimType GetType() = 0;
+		virtual bool IsEnabled() { return bEnabled; }
 		virtual FSimOutputData* MakeNewData() = 0;
-		virtual void FillOutputState(const ISimulationModuleBase* SimModule) = 0;
+		virtual void FillOutputState(const ISimulationModuleBase* SimModule);
 		virtual void Lerp(const FSimOutputData& InCurrent, const FSimOutputData& InNext, float Alpha) = 0;
+
+		bool bEnabled = true;
+
+#if !(UE_BUILD_SHIPPING || UE_BUILD_TEST)
 		virtual FString ToString() { return FString(); }
+		FString DebugString;
+#endif
 	};
 
 
diff --git a/Engine/Source/Runtime/Experimental/ChaosVehicles/ChaosVehiclesCore/Public/SimModule/SuspensionModule.h b/Engine/Source/Runtime/Experimental/ChaosVehicles/ChaosVehiclesCore/Public/SimModule/SuspensionModule.h
--- a/Engine/Source/Runtime/Experimental/ChaosVehicles/ChaosVehiclesCore/Public/SimModule/SuspensionModule.h
+++ b/Engine/Source/Runtime/Experimental/ChaosVehicles/ChaosVehiclesCore/Public/SimModule/SuspensionModule.h
@@ -50,7 +50,9 @@ namespace Chaos
 		virtual void FillOutputState(const ISimulationModuleBase* SimModule) override;
 		virtual void Lerp(const FSimOutputData& InCurrent, const FSimOutputData& InNext, float Alpha) override;
 
+#if !(UE_BUILD_SHIPPING || UE_BUILD_TEST)
 		virtual FString ToString() override;
+#endif
 
 		float SpringDisplacement;
 		float SpringSpeed;
diff --git a/Engine/Source/Runtime/Experimental/ChaosVehicles/ChaosVehiclesCore/Public/SimModule/TransmissionModule.h b/Engine/Source/Runtime/Experimental/ChaosVehicles/ChaosVehiclesCore/Public/SimModule/TransmissionModule.h
--- a/Engine/Source/Runtime/Experimental/ChaosVehicles/ChaosVehiclesCore/Public/SimModule/TransmissionModule.h
+++ b/Engine/Source/Runtime/Experimental/ChaosVehicles/ChaosVehiclesCore/Public/SimModule/TransmissionModule.h
@@ -50,7 +50,10 @@ namespace Chaos
 		virtual eSimType GetType() override { return eSimType::Transmission; }
 		virtual void FillOutputState(const ISimulationModuleBase* SimModule) override;
 		virtual void Lerp(const FSimOutputData& InCurrent, const FSimOutputData& InNext, float Alpha) override;
+
+#if !(UE_BUILD_SHIPPING || UE_BUILD_TEST)
 		virtual FString ToString() override;
+#endif
 
 		int32 CurrentGear;
 	};
diff --git a/Engine/Source/Runtime/Experimental/ChaosVehicles/ChaosVehiclesCore/Public/SimModule/WheelModule.h b/Engine/Source/Runtime/Experimental/ChaosVehicles/ChaosVehiclesCore/Public/SimModule/WheelModule.h
--- a/Engine/Source/Runtime/Experimental/ChaosVehicles/ChaosVehiclesCore/Public/SimModule/WheelModule.h
+++ b/Engine/Source/Runtime/Experimental/ChaosVehicles/ChaosVehiclesCore/Public/SimModule/WheelModule.h
@@ -48,7 +48,10 @@ namespace Chaos
 		virtual eSimType GetType() override { return eSimType::Wheel; }
 		virtual void FillOutputState(const ISimulationModuleBase* SimModule) override;
 		virtual void Lerp(const FSimOutputData& InCurrent, const FSimOutputData& InNext, float Alpha) override;
+
+#if !(UE_BUILD_SHIPPING || UE_BUILD_TEST)
 		virtual FString ToString() override;
+#endif
 
 		bool bTouchingGround;
 		float ForceIntoSurface;
