Removing the geometry collection asset pathway, keeping only the cluster union setup.

diff --git a/Engine/Plugins/Experimental/ChaosModularVehicle/ChaosModularVehicle.uplugin b/Engine/Plugins/Experimental/ChaosModularVehicle/ChaosModularVehicle.uplugin
--- a/Engine/Plugins/Experimental/ChaosModularVehicle/ChaosModularVehicle.uplugin
+++ b/Engine/Plugins/Experimental/ChaosModularVehicle/ChaosModularVehicle.uplugin
@@ -24,13 +24,7 @@
 			"Name": "ChaosModularVehicleEngine",
 			"Type": "Runtime",
 			"LoadingPhase": "PreDefault"
-		},
-		{
-			"Name": "ChaosModularVehicleEditor",
-			"Type": "Editor",
-			"LoadingPhase": "PreDefault"
 		}
-
 	],
 	"Plugins": [
 		{
diff --git a/Engine/Plugins/Experimental/ChaosModularVehicle/Source/ChaosModularVehicleEditor/ChaosModularVehicleEditor.Build.cs b/Engine/Plugins/Experimental/ChaosModularVehicle/Source/ChaosModularVehicleEditor/ChaosModularVehicleEditor.Build.cs
deleted file mode 100644
diff --git a/Engine/Plugins/Experimental/ChaosModularVehicle/Source/ChaosModularVehicleEditor/Private/ChaosModularVehicle/ActorFactoryModularVehicle.cpp b/Engine/Plugins/Experimental/ChaosModularVehicle/Source/ChaosModularVehicleEditor/Private/ChaosModularVehicle/ActorFactoryModularVehicle.cpp
deleted file mode 100644
diff --git a/Engine/Plugins/Experimental/ChaosModularVehicle/Source/ChaosModularVehicleEditor/Private/ChaosModularVehicle/ActorFactoryModularVehicle.h b/Engine/Plugins/Experimental/ChaosModularVehicle/Source/ChaosModularVehicleEditor/Private/ChaosModularVehicle/ActorFactoryModularVehicle.h
deleted file mode 100644
diff --git a/Engine/Plugins/Experimental/ChaosModularVehicle/Source/ChaosModularVehicleEditor/Private/ChaosModularVehicle/AssetTypeActions_ModularVehicleAsset.cpp b/Engine/Plugins/Experimental/ChaosModularVehicle/Source/ChaosModularVehicleEditor/Private/ChaosModularVehicle/AssetTypeActions_ModularVehicleAsset.cpp
deleted file mode 100644
diff --git a/Engine/Plugins/Experimental/ChaosModularVehicle/Source/ChaosModularVehicleEditor/Private/ChaosModularVehicle/ChaosModularVehicleCommands.cpp b/Engine/Plugins/Experimental/ChaosModularVehicle/Source/ChaosModularVehicleEditor/Private/ChaosModularVehicle/ChaosModularVehicleCommands.cpp
deleted file mode 100644
diff --git a/Engine/Plugins/Experimental/ChaosModularVehicle/Source/ChaosModularVehicleEditor/Private/ChaosModularVehicle/ChaosModularVehicleEditorPlugin.cpp b/Engine/Plugins/Experimental/ChaosModularVehicle/Source/ChaosModularVehicleEditor/Private/ChaosModularVehicle/ChaosModularVehicleEditorPlugin.cpp
deleted file mode 100644
diff --git a/Engine/Plugins/Experimental/ChaosModularVehicle/Source/ChaosModularVehicleEditor/Private/ChaosModularVehicle/ModularToolkitBlueprintLibrary.cpp b/Engine/Plugins/Experimental/ChaosModularVehicle/Source/ChaosModularVehicleEditor/Private/ChaosModularVehicle/ModularToolkitBlueprintLibrary.cpp
deleted file mode 100644
diff --git a/Engine/Plugins/Experimental/ChaosModularVehicle/Source/ChaosModularVehicleEditor/Private/ChaosModularVehicle/ModularVehicleAssetConversion.cpp b/Engine/Plugins/Experimental/ChaosModularVehicle/Source/ChaosModularVehicleEditor/Private/ChaosModularVehicle/ModularVehicleAssetConversion.cpp
deleted file mode 100644
diff --git a/Engine/Plugins/Experimental/ChaosModularVehicle/Source/ChaosModularVehicleEditor/Private/ChaosModularVehicle/ModularVehicleAssetFactory.cpp b/Engine/Plugins/Experimental/ChaosModularVehicle/Source/ChaosModularVehicleEditor/Private/ChaosModularVehicle/ModularVehicleAssetFactory.cpp
deleted file mode 100644
diff --git a/Engine/Plugins/Experimental/ChaosModularVehicle/Source/ChaosModularVehicleEditor/Private/ChaosModularVehicle/ModularVehicleAssetThumbnailRenderer.cpp b/Engine/Plugins/Experimental/ChaosModularVehicle/Source/ChaosModularVehicleEditor/Private/ChaosModularVehicle/ModularVehicleAssetThumbnailRenderer.cpp
deleted file mode 100644
diff --git a/Engine/Plugins/Experimental/ChaosModularVehicle/Source/ChaosModularVehicleEditor/Private/ChaosModularVehicle/ModularVehicleAssetThumbnailScene.cpp b/Engine/Plugins/Experimental/ChaosModularVehicle/Source/ChaosModularVehicleEditor/Private/ChaosModularVehicle/ModularVehicleAssetThumbnailScene.cpp
deleted file mode 100644
diff --git a/Engine/Plugins/Experimental/ChaosModularVehicle/Source/ChaosModularVehicleEditor/Private/ChaosModularVehicle/ModularVehicleBPInterface.cpp b/Engine/Plugins/Experimental/ChaosModularVehicle/Source/ChaosModularVehicleEditor/Private/ChaosModularVehicle/ModularVehicleBPInterface.cpp
deleted file mode 100644
diff --git a/Engine/Plugins/Experimental/ChaosModularVehicle/Source/ChaosModularVehicleEditor/Public/ChaosModularVehicle/AssetTypeActions_ModularVehicleAsset.h b/Engine/Plugins/Experimental/ChaosModularVehicle/Source/ChaosModularVehicleEditor/Public/ChaosModularVehicle/AssetTypeActions_ModularVehicleAsset.h
deleted file mode 100644
diff --git a/Engine/Plugins/Experimental/ChaosModularVehicle/Source/ChaosModularVehicleEditor/Public/ChaosModularVehicle/ChaosModularVehicleCommands.h b/Engine/Plugins/Experimental/ChaosModularVehicle/Source/ChaosModularVehicleEditor/Public/ChaosModularVehicle/ChaosModularVehicleCommands.h
deleted file mode 100644
diff --git a/Engine/Plugins/Experimental/ChaosModularVehicle/Source/ChaosModularVehicleEditor/Public/ChaosModularVehicle/ChaosModularVehicleEditorPlugin.h b/Engine/Plugins/Experimental/ChaosModularVehicle/Source/ChaosModularVehicleEditor/Public/ChaosModularVehicle/ChaosModularVehicleEditorPlugin.h
deleted file mode 100644
diff --git a/Engine/Plugins/Experimental/ChaosModularVehicle/Source/ChaosModularVehicleEditor/Public/ChaosModularVehicle/ModularToolkitBlueprintLibrary.h b/Engine/Plugins/Experimental/ChaosModularVehicle/Source/ChaosModularVehicleEditor/Public/ChaosModularVehicle/ModularToolkitBlueprintLibrary.h
deleted file mode 100644
diff --git a/Engine/Plugins/Experimental/ChaosModularVehicle/Source/ChaosModularVehicleEditor/Public/ChaosModularVehicle/ModularVehicleAssetConversion.h b/Engine/Plugins/Experimental/ChaosModularVehicle/Source/ChaosModularVehicleEditor/Public/ChaosModularVehicle/ModularVehicleAssetConversion.h
deleted file mode 100644
diff --git a/Engine/Plugins/Experimental/ChaosModularVehicle/Source/ChaosModularVehicleEditor/Public/ChaosModularVehicle/ModularVehicleAssetFactory.h b/Engine/Plugins/Experimental/ChaosModularVehicle/Source/ChaosModularVehicleEditor/Public/ChaosModularVehicle/ModularVehicleAssetFactory.h
deleted file mode 100644
diff --git a/Engine/Plugins/Experimental/ChaosModularVehicle/Source/ChaosModularVehicleEditor/Public/ChaosModularVehicle/ModularVehicleAssetThumbnailRenderer.h b/Engine/Plugins/Experimental/ChaosModularVehicle/Source/ChaosModularVehicleEditor/Public/ChaosModularVehicle/ModularVehicleAssetThumbnailRenderer.h
deleted file mode 100644
diff --git a/Engine/Plugins/Experimental/ChaosModularVehicle/Source/ChaosModularVehicleEditor/Public/ChaosModularVehicle/ModularVehicleAssetThumbnailScene.h b/Engine/Plugins/Experimental/ChaosModularVehicle/Source/ChaosModularVehicleEditor/Public/ChaosModularVehicle/ModularVehicleAssetThumbnailScene.h
deleted file mode 100644
diff --git a/Engine/Plugins/Experimental/ChaosModularVehicle/Source/ChaosModularVehicleEditor/Public/ChaosModularVehicle/ModularVehicleBPInterface.h b/Engine/Plugins/Experimental/ChaosModularVehicle/Source/ChaosModularVehicleEditor/Public/ChaosModularVehicle/ModularVehicleBPInterface.h
deleted file mode 100644
diff --git a/Engine/Plugins/Experimental/ChaosModularVehicle/Source/ChaosModularVehicleEngine/Private/ChaosModularVehicle/ChaosSimModuleManager.cpp b/Engine/Plugins/Experimental/ChaosModularVehicle/Source/ChaosModularVehicleEngine/Private/ChaosModularVehicle/ChaosSimModuleManager.cpp
--- a/Engine/Plugins/Experimental/ChaosModularVehicle/Source/ChaosModularVehicleEngine/Private/ChaosModularVehicle/ChaosSimModuleManager.cpp
+++ b/Engine/Plugins/Experimental/ChaosModularVehicle/Source/ChaosModularVehicleEngine/Private/ChaosModularVehicle/ChaosSimModuleManager.cpp
@@ -1,7 +1,6 @@
 //// Copyright Epic Games, Inc. All Rights Reserved.
 //
 #include "ChaosModularVehicle/ChaosSimModuleManager.h"
-#include "ChaosModularVehicle/ModularVehicleComponent.h"
 #include "ChaosModularVehicle/ModularVehicleBaseComponent.h"
 #include "PBDRigidsSolver.h"
 #include "GameFramework/HUD.h" // for ShowDebugInfo
@@ -49,11 +48,6 @@ FChaosSimModuleManager::FChaosSimModuleManager(FPhysScene* PhysScene)
 
 FChaosSimModuleManager::~FChaosSimModuleManager()
 {
-	while (GCVehicles.Num() > 0)
-	{
-		RemoveVehicle(GCVehicles.Last());
-	}
-
 	while (CUVehicles.Num() > 0)
 	{
 		RemoveVehicle(CUVehicles.Last());
@@ -93,16 +87,6 @@ void FChaosSimModuleManager::OnShowDebugInfo(AHUD* HUD, UCanvas* Canvas, const F
 				Manager->CUVehicles[ShowVehicleIndex]->ShowDebugInfo(HUD, Canvas, DisplayInfo, YL, YPos);
 			}
 		}
-
-		if (!Manager->GCVehicles.IsEmpty())
-		{
-			if (Manager->GCVehicles[ShowVehicleIndex].IsValid())
-			{
-				Manager->GCVehicles[ShowVehicleIndex]->ShowDebugInfo(HUD, Canvas, DisplayInfo, YL, YPos);
-			}
-		}
-
-
 	}
 }
 
@@ -156,23 +140,7 @@ FChaosSimModuleManager* FChaosSimModuleManager::GetManagerFromScene(FPhysScene*
 	return Manager;
 }
 
-void FChaosSimModuleManager::AddVehicle(TWeakObjectPtr<UModularVehicleComponent> Vehicle)
-{
-	check(Vehicle != NULL);
-	check(Vehicle->PhysicsVehicleOutput());
-	check(AsyncCallback);
-
-	GCVehicles.Add(Vehicle);
-}
-
-void FChaosSimModuleManager::RemoveVehicle(TWeakObjectPtr<UModularVehicleComponent> Vehicle)
-{
-	check(Vehicle != NULL);
-
-	GCVehicles.Remove(Vehicle);
-}
-
-void FChaosSimModuleManager::AddVehicle(TWeakObjectPtr < UModularVehicleBaseComponent> Vehicle)
+void FChaosSimModuleManager::AddVehicle(TWeakObjectPtr<UModularVehicleBaseComponent> Vehicle)
 {
 	check(Vehicle != NULL);
 	check(Vehicle->PhysicsVehicleOutput());
@@ -189,11 +157,6 @@ void FChaosSimModuleManager::RemoveVehicle(TWeakObjectPtr<UModularVehicleBaseCom
 
 void FChaosSimModuleManager::ScenePreTick(FPhysScene* PhysScene, float DeltaTime)
 {
-	for (int32 i = 0; i < GCVehicles.Num(); ++i)
-	{
-		GCVehicles[i]->PreTickGT(DeltaTime);
-	}
-
 	for (int32 i = 0; i < CUVehicles.Num(); ++i)
 	{
 		CUVehicles[i]->PreTickGT(DeltaTime);
@@ -214,13 +177,6 @@ void FChaosSimModuleManager::Update(FPhysScene* PhysScene, float DeltaTime)
 	{
 		FChaosSimModuleManagerAsyncInput* AsyncInput = AsyncCallback->GetProducerInputData_External();
 
-		for (TWeakObjectPtr<UModularVehicleComponent> Vehicle : GCVehicles)
-		{
-			Vehicle->Update(DeltaTime);
-			Vehicle->FinalizeSimCallbackData(*AsyncInput);
-		}
-
-
 		for (TWeakObjectPtr<UModularVehicleBaseComponent> Vehicle : CUVehicles)
 		{
 			Vehicle->Update(DeltaTime);
@@ -297,25 +253,6 @@ void FChaosSimModuleManager::ParallelUpdateVehicles(float DeltaSeconds)
 
 	if (UWorld* World = Scene.GetOwningWorld())
 	{
-		{
-		int32 NumVehiclesInActiveBatch = 0;
-		for (TWeakObjectPtr<UModularVehicleComponent> Vehicle : GCVehicles)
-		{
-			auto NextOutput = PendingOutputs.Num() > 0 ? PendingOutputs[0].Get() : nullptr;
-			float Alpha = 0.f;
-			if (NextOutput && LatestOutput)
-			{
-				const float Denom = NextOutput->InternalTime - LatestOutput->InternalTime;
-				if (Denom > SMALL_NUMBER)
-				{
-					Alpha = (ResultsTime - LatestOutput->InternalTime) / Denom;
-				}
-			}
-
-			AsyncInput->VehicleInputs.Add(Vehicle->SetCurrentAsyncData(AsyncInput->VehicleInputs.Num(), LatestOutput.Get(), NextOutput, Alpha, Timestamp));
-		}
-
-		}
 		int32 NumVehiclesInActiveBatch = 0;
 
 		for (TWeakObjectPtr<UModularVehicleBaseComponent> Vehicle : CUVehicles)
@@ -340,19 +277,6 @@ void FChaosSimModuleManager::ParallelUpdateVehicles(float DeltaSeconds)
 
 	bool ForceSingleThread = !GSimModuleDebugParams.EnableMultithreading;
 
-
-	{
-		const auto& AwakeVehiclesBatch = GCVehicles;
-		auto LambdaParallelUpdate = [DeltaSeconds, &AwakeVehiclesBatch](int32 Idx)
-		{
-			TWeakObjectPtr<UModularVehicleComponent> Vehicle = AwakeVehiclesBatch[Idx];
-			Vehicle->ParallelUpdate(DeltaSeconds); // gets output state from PT
-		};
-
-		ParallelFor(AwakeVehiclesBatch.Num(), LambdaParallelUpdate, ForceSingleThread);
-	}
-
-
 	{
 		const auto& AwakeVehiclesBatch = CUVehicles;
 		auto LambdaParallelUpdate2 = [DeltaSeconds, &AwakeVehiclesBatch](int32 Idx)
diff --git a/Engine/Plugins/Experimental/ChaosModularVehicle/Source/ChaosModularVehicleEngine/Private/ChaosModularVehicle/ChaosSimModuleManagerAsyncCallback.cpp b/Engine/Plugins/Experimental/ChaosModularVehicle/Source/ChaosModularVehicleEngine/Private/ChaosModularVehicle/ChaosSimModuleManagerAsyncCallback.cpp
--- a/Engine/Plugins/Experimental/ChaosModularVehicle/Source/ChaosModularVehicleEngine/Private/ChaosModularVehicle/ChaosSimModuleManagerAsyncCallback.cpp
+++ b/Engine/Plugins/Experimental/ChaosModularVehicle/Source/ChaosModularVehicleEngine/Private/ChaosModularVehicle/ChaosSimModuleManagerAsyncCallback.cpp
@@ -2,7 +2,6 @@
 
 #include "ChaosModularVehicle/ChaosSimModuleManagerAsyncCallback.h"
 
-#include "ChaosModularVehicle/ModularVehicleComponent.h"
 #include "ChaosModularVehicle/ModularVehicleBaseComponent.h"
 #include "ChaosModularVehicle/ModularVehicleSimulationCU.h"
 #include "PBDRigidsSolver.h"
diff --git a/Engine/Plugins/Experimental/ChaosModularVehicle/Source/ChaosModularVehicleEngine/Private/ChaosModularVehicle/ModularVehicleActor.cpp b/Engine/Plugins/Experimental/ChaosModularVehicle/Source/ChaosModularVehicleEngine/Private/ChaosModularVehicle/ModularVehicleActor.cpp
deleted file mode 100644
diff --git a/Engine/Plugins/Experimental/ChaosModularVehicle/Source/ChaosModularVehicleEngine/Private/ChaosModularVehicle/ModularVehicleAsset.cpp b/Engine/Plugins/Experimental/ChaosModularVehicle/Source/ChaosModularVehicleEngine/Private/ChaosModularVehicle/ModularVehicleAsset.cpp
--- a/Engine/Plugins/Experimental/ChaosModularVehicle/Source/ChaosModularVehicleEngine/Private/ChaosModularVehicle/ModularVehicleAsset.cpp
+++ b/Engine/Plugins/Experimental/ChaosModularVehicle/Source/ChaosModularVehicleEngine/Private/ChaosModularVehicle/ModularVehicleAsset.cpp
@@ -5,7 +5,6 @@
 =============================================================================*/
 #include "ChaosModularVehicle/ModularVehicleAsset.h"
 #include "ChaosModularVehicle/ModularSimCollection.h"
-#include "ChaosModularVehicle/ModularVehicleComponent.h"
 #include "UObject/Package.h"
 
 DEFINE_LOG_CATEGORY_STATIC(LogModularVehicleAssetInternal, Log, All);
diff --git a/Engine/Plugins/Experimental/ChaosModularVehicle/Source/ChaosModularVehicleEngine/Private/ChaosModularVehicle/ModularVehicleClusterActor.cpp b/Engine/Plugins/Experimental/ChaosModularVehicle/Source/ChaosModularVehicleEngine/Private/ChaosModularVehicle/ModularVehicleClusterActor.cpp
--- a/Engine/Plugins/Experimental/ChaosModularVehicle/Source/ChaosModularVehicleEngine/Private/ChaosModularVehicle/ModularVehicleClusterActor.cpp
+++ b/Engine/Plugins/Experimental/ChaosModularVehicle/Source/ChaosModularVehicleEngine/Private/ChaosModularVehicle/ModularVehicleClusterActor.cpp
@@ -1,20 +1,19 @@
 // Copyright Epic Games, Inc. All Rights Reserved.
 
 #include "ChaosModularVehicle/ModularVehicleClusterActor.h"
-
-#include "PhysicsEngine/ClusterUnionComponent.h"
+#include "ChaosModularVehicle/ClusterUnionVehicleComponent.h"
 #include "ChaosModularVehicle/ModularVehicleBaseComponent.h"
 
 
 AModularVehicleClusterActor::AModularVehicleClusterActor(const FObjectInitializer& ObjectInitializer)
 	: Super(ObjectInitializer)
 {
-	ClusterUnionComponent = CreateDefaultSubobject<UClusterUnionComponent>(TEXT("ClusterUnionComponent0"));
-	SetRootComponent(ClusterUnionComponent);
+	ClusterUnionVehicleComponent = CreateDefaultSubobject<UClusterUnionVehicleComponent>(TEXT("ClusterUnionVehicleComponent0"));
+	SetRootComponent(ClusterUnionVehicleComponent);
 
 	// Vehicle sim piggy-backs off of the cluster union component & its events
 	VehicleSimComponent = CreateDefaultSubobject<UModularVehicleBaseComponent>(TEXT("VehicleSimComponent0"));
-	VehicleSimComponent->SetClusterComponent(ClusterUnionComponent);
+	VehicleSimComponent->SetClusterComponent(ClusterUnionVehicleComponent);
 
 	SetRemoteRoleForBackwardsCompat(ROLE_SimulatedProxy);
 	bReplicates = true;
diff --git a/Engine/Plugins/Experimental/ChaosModularVehicle/Source/ChaosModularVehicleEngine/Private/ChaosModularVehicle/ModularVehicleComponent.cpp b/Engine/Plugins/Experimental/ChaosModularVehicle/Source/ChaosModularVehicleEngine/Private/ChaosModularVehicle/ModularVehicleComponent.cpp
deleted file mode 100644
diff --git a/Engine/Plugins/Experimental/ChaosModularVehicle/Source/ChaosModularVehicleEngine/Private/ChaosModularVehicle/ModularVehicleObject.cpp b/Engine/Plugins/Experimental/ChaosModularVehicle/Source/ChaosModularVehicleEngine/Private/ChaosModularVehicle/ModularVehicleObject.cpp
deleted file mode 100644
diff --git a/Engine/Plugins/Experimental/ChaosModularVehicle/Source/ChaosModularVehicleEngine/Private/ChaosModularVehicle/ModularVehiclePawn.cpp b/Engine/Plugins/Experimental/ChaosModularVehicle/Source/ChaosModularVehicleEngine/Private/ChaosModularVehicle/ModularVehiclePawn.cpp
deleted file mode 100644
diff --git a/Engine/Plugins/Experimental/ChaosModularVehicle/Source/ChaosModularVehicleEngine/Private/ChaosModularVehicle/ModularVehicleSimulationGC.cpp b/Engine/Plugins/Experimental/ChaosModularVehicle/Source/ChaosModularVehicleEngine/Private/ChaosModularVehicle/ModularVehicleSimulationGC.cpp
deleted file mode 100644
diff --git a/Engine/Plugins/Experimental/ChaosModularVehicle/Source/ChaosModularVehicleEngine/Public/ChaosModularVehicle/ChaosSimModuleManager.h b/Engine/Plugins/Experimental/ChaosModularVehicle/Source/ChaosModularVehicleEngine/Public/ChaosModularVehicle/ChaosSimModuleManager.h
--- a/Engine/Plugins/Experimental/ChaosModularVehicle/Source/ChaosModularVehicleEngine/Public/ChaosModularVehicle/ChaosSimModuleManager.h
+++ b/Engine/Plugins/Experimental/ChaosModularVehicle/Source/ChaosModularVehicleEngine/Public/ChaosModularVehicle/ChaosSimModuleManager.h
@@ -35,13 +35,11 @@ public:
 	/**
 	 * Register a Physics vehicle for processing
 	 */
-	void AddVehicle(TWeakObjectPtr<UModularVehicleComponent> Vehicle);
 	void AddVehicle(TWeakObjectPtr<UModularVehicleBaseComponent> Vehicle);
 
 	/**
 	 * Unregister a Physics vehicle from processing
 	 */
-	void RemoveVehicle(TWeakObjectPtr<UModularVehicleComponent> Vehicle);
 	void RemoveVehicle(TWeakObjectPtr<UModularVehicleBaseComponent> Vehicle);
 
 	/**
diff --git a/Engine/Plugins/Experimental/ChaosModularVehicle/Source/ChaosModularVehicleEngine/Public/ChaosModularVehicle/ModularVehicleActor.h b/Engine/Plugins/Experimental/ChaosModularVehicle/Source/ChaosModularVehicleEngine/Public/ChaosModularVehicle/ModularVehicleActor.h
deleted file mode 100644
diff --git a/Engine/Plugins/Experimental/ChaosModularVehicle/Source/ChaosModularVehicleEngine/Public/ChaosModularVehicle/ModularVehicleClusterActor.h b/Engine/Plugins/Experimental/ChaosModularVehicle/Source/ChaosModularVehicleEngine/Public/ChaosModularVehicle/ModularVehicleClusterActor.h
--- a/Engine/Plugins/Experimental/ChaosModularVehicle/Source/ChaosModularVehicleEngine/Public/ChaosModularVehicle/ModularVehicleClusterActor.h
+++ b/Engine/Plugins/Experimental/ChaosModularVehicle/Source/ChaosModularVehicleEngine/Public/ChaosModularVehicle/ModularVehicleClusterActor.h
@@ -9,7 +9,7 @@
 
 #include "ModularVehicleClusterActor.generated.h"
 
-class UClusterUnionComponent;
+class UClusterUnionVehicleComponent;
 class UModularVehicleBaseComponent;
 
 UCLASS()
@@ -18,10 +18,15 @@ class CHAOSMODULARVEHICLEENGINE_API AModularVehicleClusterActor: public AActor
 	GENERATED_UCLASS_BODY()
 
 public:
+	UFUNCTION()
+	UClusterUnionVehicleComponent* GetClusterUnionComponent() const { return ClusterUnionVehicleComponent; }
+
+	UFUNCTION()
+	UModularVehicleBaseComponent* GetVehicleSimulationComponent() const { return VehicleSimComponent; }
 
 	/* VehicleSpecificClusterComponent */
 	UPROPERTY(VisibleAnywhere, BlueprintReadOnly, Category = Vehicle, meta = (ExposeFunctionCategories = "Components|ModularVehicle", AllowPrivateAccess = "true"))
-	TObjectPtr<UClusterUnionComponent> ClusterUnionComponent;
+	TObjectPtr<UClusterUnionVehicleComponent> ClusterUnionVehicleComponent;
 
 	/* ModularVehicleComponent */
 	UPROPERTY(VisibleAnywhere, BlueprintReadOnly, Category = Vehicle, meta = (ExposeFunctionCategories = "Components|ModularVehicle", AllowPrivateAccess = "true"))
diff --git a/Engine/Plugins/Experimental/ChaosModularVehicle/Source/ChaosModularVehicleEngine/Public/ChaosModularVehicle/ModularVehicleComponent.h b/Engine/Plugins/Experimental/ChaosModularVehicle/Source/ChaosModularVehicleEngine/Public/ChaosModularVehicle/ModularVehicleComponent.h
deleted file mode 100644
diff --git a/Engine/Plugins/Experimental/ChaosModularVehicle/Source/ChaosModularVehicleEngine/Public/ChaosModularVehicle/ModularVehicleObject.h b/Engine/Plugins/Experimental/ChaosModularVehicle/Source/ChaosModularVehicleEngine/Public/ChaosModularVehicle/ModularVehicleObject.h
deleted file mode 100644
diff --git a/Engine/Plugins/Experimental/ChaosModularVehicle/Source/ChaosModularVehicleEngine/Public/ChaosModularVehicle/ModularVehiclePawn.h b/Engine/Plugins/Experimental/ChaosModularVehicle/Source/ChaosModularVehicleEngine/Public/ChaosModularVehicle/ModularVehiclePawn.h
deleted file mode 100644
diff --git a/Engine/Plugins/Experimental/ChaosModularVehicle/Source/ChaosModularVehicleEngine/Public/ChaosModularVehicle/ModularVehicleSimulationGC.h b/Engine/Plugins/Experimental/ChaosModularVehicle/Source/ChaosModularVehicleEngine/Public/ChaosModularVehicle/ModularVehicleSimulationGC.h
deleted file mode 100644
diff --git a/Engine/Source/Runtime/Experimental/ChaosVehicles/ChaosVehiclesCore/Private/SimModule/SimModuleTree.cpp b/Engine/Source/Runtime/Experimental/ChaosVehicles/ChaosVehiclesCore/Private/SimModule/SimModuleTree.cpp
--- a/Engine/Source/Runtime/Experimental/ChaosVehicles/ChaosVehiclesCore/Private/SimModule/SimModuleTree.cpp
+++ b/Engine/Source/Runtime/Experimental/ChaosVehicles/ChaosVehiclesCore/Private/SimModule/SimModuleTree.cpp
@@ -235,26 +235,6 @@ void FSimModuleTree::Simulate(float DeltaTime, FAllInputs& Inputs, FClusterUnion
 
 }
 
-
-void FSimModuleTree::Simulate(float DeltaTime, FAllInputs& Inputs, FGeometryCollectionPhysicsProxy* PhysicsProxy)
-{
-	SCOPE_CYCLE_COUNTER(STAT_ModularVehicle_SimulateTree);
-
-	if (PhysicsProxy)
-	{
-		UpdateModuleVelocites(PhysicsProxy);
-	}
-
-	TArray<int> RootNodes;
-	GetRootNodes(RootNodes);
-
-	for (int RootIndex : RootNodes)
-	{
-		SimulateNode(DeltaTime, Inputs, RootIndex, nullptr);
-	}
-
-}
-
 void FSimModuleTree::SimulateNode(float DeltaTime, FAllInputs& Inputs, int NodeIndex, FClusterUnionPhysicsProxy* PhysicsProxy)
 {
 	if (ISimulationModuleBase* Module = AccessSimModule(NodeIndex))
@@ -276,7 +256,6 @@ void FSimModuleTree::SimulateNode(float DeltaTime, FAllInputs& Inputs, int NodeI
 	}
 }
 
-
 void FSimModuleTree::DeleteNodesBelow(int AtIndex)
 {
 	if (IsValidNode(AtIndex))
diff --git a/Engine/Source/Runtime/Experimental/ChaosVehicles/ChaosVehiclesCore/Public/SimModule/SimModuleTree.h b/Engine/Source/Runtime/Experimental/ChaosVehicles/ChaosVehiclesCore/Public/SimModule/SimModuleTree.h
--- a/Engine/Source/Runtime/Experimental/ChaosVehicles/ChaosVehiclesCore/Public/SimModule/SimModuleTree.h
+++ b/Engine/Source/Runtime/Experimental/ChaosVehicles/ChaosVehiclesCore/Public/SimModule/SimModuleTree.h
@@ -138,7 +138,6 @@ namespace Chaos
 
 		void AppendTreeUpdates(const FSimTreeUpdates& TreeUpdates);
 
-		void Simulate(float DeltaTime, FAllInputs& Inputs, FGeometryCollectionPhysicsProxy* PhysicsProxy);
 		void Simulate(float DeltaTime, FAllInputs& Inputs, FClusterUnionPhysicsProxy* PhysicsProxy);
 
 		FDeferredForcesModular& AccessDeferredForces() { return DeferredForces; }
