Implement the Locally controlled without need for pawn, various bug fixes.

diff --git a/Engine/Plugins/Experimental/ChaosModularVehicle/Source/ChaosModularVehicleEngine/Private/ChaosModularVehicle/ChaosModularVehicleEnginePlugin.cpp b/Engine/Plugins/Experimental/ChaosModularVehicle/Source/ChaosModularVehicleEngine/Private/ChaosModularVehicle/ChaosModularVehicleEnginePlugin.cpp
--- a/Engine/Plugins/Experimental/ChaosModularVehicle/Source/ChaosModularVehicleEngine/Private/ChaosModularVehicle/ChaosModularVehicleEnginePlugin.cpp
+++ b/Engine/Plugins/Experimental/ChaosModularVehicle/Source/ChaosModularVehicleEngine/Private/ChaosModularVehicle/ChaosModularVehicleEnginePlugin.cpp
@@ -45,4 +45,4 @@ class FChaosModularVehicleEnginePlugin : public IChaosModularVehicleEnginePlugin
 
 };
 
-IMPLEMENT_MODULE(FChaosModularVehicleEnginePlugin, ModularVehicleEngine)
+IMPLEMENT_MODULE(FChaosModularVehicleEnginePlugin, ChaosModularVehicleEngine)
diff --git a/Engine/Plugins/Experimental/ChaosModularVehicle/Source/ChaosModularVehicleEngine/Private/ChaosModularVehicle/ChaosSimModuleManagerAsyncCallback.cpp b/Engine/Plugins/Experimental/ChaosModularVehicle/Source/ChaosModularVehicleEngine/Private/ChaosModularVehicle/ChaosSimModuleManagerAsyncCallback.cpp
--- a/Engine/Plugins/Experimental/ChaosModularVehicle/Source/ChaosModularVehicleEngine/Private/ChaosModularVehicle/ChaosSimModuleManagerAsyncCallback.cpp
+++ b/Engine/Plugins/Experimental/ChaosModularVehicle/Source/ChaosModularVehicleEngine/Private/ChaosModularVehicle/ChaosSimModuleManagerAsyncCallback.cpp
@@ -178,8 +178,7 @@ void FModularVehicleAsyncInput::ProcessInputs()
 		}
 	}
 
-	APlayerController* PlayerController = GetVehicle()->GetPlayerController();
-	if (PlayerController && PlayerController->IsLocalController() && !bIsResimming)
+	if (GetVehicle()->IsLocallyControlled() && !bIsResimming)
 	{
 		VehicleSim->VehicleInputs = PhysicsInputs.NetworkInputs.VehicleInputs;
 	}
diff --git a/Engine/Plugins/Experimental/ChaosModularVehicle/Source/ChaosModularVehicleEngine/Private/ChaosModularVehicle/ModularVehicleBaseComponent.cpp b/Engine/Plugins/Experimental/ChaosModularVehicle/Source/ChaosModularVehicleEngine/Private/ChaosModularVehicle/ModularVehicleBaseComponent.cpp
--- a/Engine/Plugins/Experimental/ChaosModularVehicle/Source/ChaosModularVehicleEngine/Private/ChaosModularVehicle/ModularVehicleBaseComponent.cpp
+++ b/Engine/Plugins/Experimental/ChaosModularVehicle/Source/ChaosModularVehicleEngine/Private/ChaosModularVehicle/ModularVehicleBaseComponent.cpp
@@ -77,6 +77,7 @@ UModularVehicleBaseComponent::UModularVehicleBaseComponent(const FObjectInitiali
 		NetworkPhysicsComponent->SetIsReplicated(true);
 	}
 
+	bIsLocallyControlled = false;
 }
 
 UModularVehicleBaseComponent::~UModularVehicleBaseComponent()
@@ -106,6 +107,22 @@ APlayerController* UModularVehicleBaseComponent::GetPlayerController() const
 	return nullptr;
 }
 
+
+bool UModularVehicleBaseComponent::IsLocallyControlled() const
+{
+	if (bIsLocallyControlled && !GetWorld()->IsNetMode(NM_DedicatedServer))
+	{
+		return true;
+	}
+
+	if (APlayerController* PlayerController = GetPlayerController())
+	{
+		return PlayerController->IsLocalController();
+	}
+	return false;
+}
+
+
 void UModularVehicleBaseComponent::OnCreatePhysicsState()
 {
 	Super::OnCreatePhysicsState();
@@ -136,6 +153,12 @@ void UModularVehicleBaseComponent::OnCreatePhysicsState()
 			// register interface to handle network prediction callbacks
 			// #Note: in our case we don't yet know what the replication data will be since the modules are built after this point at runtime
 			NetworkPhysicsComponent->CreateDataHistory<FPhysicsModularVehicleTraits>(this);
+
+			if (bIsLocallyControlled)
+			{
+				NetworkPhysicsComponent->SetIsRelayingLocalInputs(bIsLocallyControlled);
+			}
+
 		}
 	}
 
@@ -457,24 +480,15 @@ void UModularVehicleBaseComponent::PreTickGT(float DeltaTime)
 		}
 	}
 
-	if (PVehicleOutput && UpdatedComponent)
-	{
-		APawn* MyOwner = Cast<APawn>(UpdatedComponent->GetOwner());
-		if (MyOwner)
-		{
-			// process control inputs and other data
-			UpdateState(DeltaTime);
-		}
-	}
+	// process control inputs and other data
+	UpdateState(DeltaTime);
 
 }
 
 void UModularVehicleBaseComponent::UpdateState(float DeltaTime)
 {
 	// update input values
-	AController* Controller = GetPlayerController();
-
-	bool bProcessLocally = bRequiresControllerForInputs ? (Controller && Controller->IsLocalController()) : true;
+	bool bProcessLocally = IsLocallyControlled();
 
 	// IsLocallyControlled will fail if the owner is unpossessed (i.e. Controller == nullptr);
 	// Should we remove input instead of relying on replicated state in that case?
@@ -849,10 +863,10 @@ void UModularVehicleBaseComponent::RemoveComponentFromSimulation(UPrimitiveCompo
 
 		Chaos::FSimTreeUpdates LatestTreeUpdates;
 
-		TArray<UActorComponent*> Components;
-		InComponent->GetOwner()->GetComponents(UVehicleSimBaseComponent::StaticClass(), Components, true);
+		TArray<USceneComponent*> Components;
+		InComponent->GetChildrenComponents(true, Components);
 
-		for (UActorComponent* ComponentPart : Components)
+		for (USceneComponent* ComponentPart : Components)
 		{
 			if (UVehicleSimBaseComponent* ChangedComponent = Cast<UVehicleSimBaseComponent>(ComponentPart))
 			{
@@ -879,7 +893,7 @@ void UModularVehicleBaseComponent::RemoveComponentFromSimulation(UPrimitiveCompo
 						if (SimModule->GetSimType() == Chaos::eSimType::Suspension)
 						{
 							Chaos::FSuspensionSimModule* SuspensionModule = static_cast<Chaos::FSuspensionSimModule*>(SimModule);
-							EnableConstraint(SuspensionModule->GetConstraintIndex(), false);
+							DestroyConstraint(SuspensionModule->GetConstraintIndex());
 						}
 
 						break;
@@ -899,6 +913,26 @@ void UModularVehicleBaseComponent::RemoveComponentFromSimulation(UPrimitiveCompo
 					SimModuleTree->AppendTreeUpdates(LatestTreeUpdates);
 				}
 			});
+
+		NextTransformIndex--;
+	}
+
+}
+
+void UModularVehicleBaseComponent::SetLocallyControlled(bool bLocallyControlledIn)
+{
+	bIsLocallyControlled = false;
+	if (UWorld* World = GetWorld())
+	{
+		if (!World->IsNetMode(NM_DedicatedServer))
+		{
+			bIsLocallyControlled = bLocallyControlledIn;
+		}
+	}
+
+	if (bUsingNetworkPhysicsPrediction && NetworkPhysicsComponent)
+	{
+		NetworkPhysicsComponent->SetIsRelayingLocalInputs(bLocallyControlledIn);
 	}
 
 }
diff --git a/Engine/Plugins/Experimental/ChaosModularVehicle/Source/ChaosModularVehicleEngine/Private/ChaosModularVehicle/ModularVehicleSimulationCU.cpp b/Engine/Plugins/Experimental/ChaosModularVehicle/Source/ChaosModularVehicleEngine/Private/ChaosModularVehicle/ModularVehicleSimulationCU.cpp
--- a/Engine/Plugins/Experimental/ChaosModularVehicle/Source/ChaosModularVehicleEngine/Private/ChaosModularVehicle/ModularVehicleSimulationCU.cpp
+++ b/Engine/Plugins/Experimental/ChaosModularVehicle/Source/ChaosModularVehicleEngine/Private/ChaosModularVehicle/ModularVehicleSimulationCU.cpp
@@ -167,12 +167,6 @@ void FModularVehicleSimulationCU::PerformAdditionalSimWork(UWorld* InWorld, cons
 				if (Node.IsValid() && Node.SimModule && Node.SimModule->IsEnabled())
 				{
 					FRigidTransform3 Frame = FRigidTransform3::Identity;
-					int TransformIndex = Node.SimModule->GetTransformIndex();
-
-					if (TransformIndex < 0 || TransformIndex >= Particles.Num())
-					{
-						continue;
-					}
 
 					FPBDRigidParticleHandle* Child = Node.SimModule->GetParticleFromUniqueIndex(Node.SimModule->GetParticleIndex().Idx, Particles);
 					if (Child == nullptr)
diff --git a/Engine/Plugins/Experimental/ChaosModularVehicle/Source/ChaosModularVehicleEngine/Private/ChaosModularVehicle/ModularVehicleSocket.cpp b/Engine/Plugins/Experimental/ChaosModularVehicle/Source/ChaosModularVehicleEngine/Private/ChaosModularVehicle/ModularVehicleSocket.cpp
--- a/Engine/Plugins/Experimental/ChaosModularVehicle/Source/ChaosModularVehicleEngine/Private/ChaosModularVehicle/ModularVehicleSocket.cpp
+++ b/Engine/Plugins/Experimental/ChaosModularVehicle/Source/ChaosModularVehicleEngine/Private/ChaosModularVehicle/ModularVehicleSocket.cpp
@@ -3,6 +3,11 @@
 #include "ChaosModularVehicle/ModularVehicleSocket.h"
 #include "ChaosModularVehicle/ClusterUnionVehicleComponent.h"
 
+FModularVehicleSocket::FModularVehicleSocket()
+{
+	RelativeLocation = FVector::ZeroVector;
+	RelativeRotation = FRotator::ZeroRotator;
+}
 
 FVector FModularVehicleSocket::GetLocation(const class UClusterUnionVehicleComponent* Component) const
 {
diff --git a/Engine/Plugins/Experimental/ChaosModularVehicle/Source/ChaosModularVehicleEngine/Public/ChaosModularVehicle/ModularVehicleBaseComponent.h b/Engine/Plugins/Experimental/ChaosModularVehicle/Source/ChaosModularVehicleEngine/Public/ChaosModularVehicle/ModularVehicleBaseComponent.h
--- a/Engine/Plugins/Experimental/ChaosModularVehicle/Source/ChaosModularVehicleEngine/Public/ChaosModularVehicle/ModularVehicleBaseComponent.h
+++ b/Engine/Plugins/Experimental/ChaosModularVehicle/Source/ChaosModularVehicleEngine/Public/ChaosModularVehicle/ModularVehicleBaseComponent.h
@@ -170,12 +170,14 @@ class CHAOSMODULARVEHICLEENGINE_API UModularVehicleBaseComponent : public UPawnM
 	friend class FModularVehicleBuilder;
 public:
 	APlayerController* GetPlayerController() const;
+	bool IsLocallyControlled() const;
 
 	virtual void GetLifetimeReplicatedProps(TArray<FLifetimeProperty>& OutLifetimeProps) const override;
 
 	virtual bool ShouldCreatePhysicsState() const override { return true; }
 	virtual void OnCreatePhysicsState() override;
 	virtual void OnDestroyPhysicsState() override;
+	virtual void SetClusterComponent(UClusterUnionComponent* InPhysicalComponent);
 
 	virtual void BeginPlay() override;
 	virtual void TickComponent(float DeltaTime, ELevelTick TickType, FActorComponentTickFunction* ThisTickFunction) override;
@@ -191,7 +193,6 @@ public:
 	void UpdateState(float DeltaTime);
 	TUniquePtr<FModularVehicleAsyncInput> SetCurrentAsyncData(int32 InputIdx, FChaosSimModuleManagerAsyncOutput* CurOutput, FChaosSimModuleManagerAsyncOutput* NextOutput, float Alpha, int32 VehicleManagerTimestamp);
 
-	void SetClusterComponent(UClusterUnionComponent* InPhysicalComponent);
 	void ParallelUpdate(float DeltaTime);
 	void Update(float DeltaTime);
 	void FinalizeSimCallbackData(FChaosSimModuleManagerAsyncInput& Input);
@@ -236,6 +237,9 @@ public:
 	UFUNCTION()
 	void RemoveComponentFromSimulation(UPrimitiveComponent* Component, const TArray<FClusterUnionBoneData>& RemovedBonesData);
 
+	UFUNCTION(BlueprintCallable, Category = "Game|Components|ModularVehicle")
+	void SetLocallyControlled(bool bLocallyControlledIn);
+
 	// CONTROLS
 	// 
 	/** Set the user input for the vehicle throttle [range 0 to 1] */
@@ -546,5 +550,7 @@ private:
 
 	TArray<FPhysicsConstraintHandle> ConstraintHandles;
 	int32 ClusteringCount = 0;
+
+	bool bIsLocallyControlled;
 };
 
diff --git a/Engine/Plugins/Experimental/ChaosModularVehicle/Source/ChaosModularVehicleEngine/Public/ChaosModularVehicle/ModularVehicleSocket.h b/Engine/Plugins/Experimental/ChaosModularVehicle/Source/ChaosModularVehicleEngine/Public/ChaosModularVehicle/ModularVehicleSocket.h
--- a/Engine/Plugins/Experimental/ChaosModularVehicle/Source/ChaosModularVehicleEngine/Public/ChaosModularVehicle/ModularVehicleSocket.h
+++ b/Engine/Plugins/Experimental/ChaosModularVehicle/Source/ChaosModularVehicleEngine/Public/ChaosModularVehicle/ModularVehicleSocket.h
@@ -16,6 +16,8 @@ struct CHAOSMODULARVEHICLEENGINE_API FModularVehicleSocket
 {
 	GENERATED_USTRUCT_BODY()
 
+	FModularVehicleSocket();
+
 	/**
 	 *	Defines a named attachment location on the Modular vehicle.
 	 */
diff --git a/Engine/Source/Runtime/Experimental/ChaosVehicles/ChaosVehiclesCore/Private/SimModule/WheelModule.cpp b/Engine/Source/Runtime/Experimental/ChaosVehicles/ChaosVehiclesCore/Private/SimModule/WheelModule.cpp
--- a/Engine/Source/Runtime/Experimental/ChaosVehicles/ChaosVehiclesCore/Private/SimModule/WheelModule.cpp
+++ b/Engine/Source/Runtime/Experimental/ChaosVehicles/ChaosVehiclesCore/Private/SimModule/WheelModule.cpp
@@ -219,7 +219,7 @@ namespace Chaos
 			if (FPBDRigidClusteredParticleHandle* ClusterChild = GetClusterParticle(Proxy))
 			{
 				float Direction = Setup().ReverseDirection ? -1.0f : 1.0f;
-				FQuat Rot = (Setup().Axis == Chaos::EWheelAxis::Y) ? FQuat(FVector(1, 0, 0), -GetAngularPosition() * Direction) : FQuat(FVector(1, 0, 0), GetAngularPosition() * Direction);
+				FQuat Rot = (Setup().Axis == Chaos::EWheelAxis::Y) ? FQuat(FVector(1, 0, 0), -GetAngularPosition() * Direction) : FQuat(FVector(0, 1, 0), GetAngularPosition() * Direction);
 				FQuat Steer = FQuat(FVector(0, 0, 1), FMath::DegreesToRadians(GetSteerAngleDegrees()));
 
 				FTransform InitialTransform = GetInitialParticleTransform();
