CIS Fix - static analysis

diff --git a/Engine/Plugins/Experimental/ChaosModularVehicle/Source/ChaosModularVehicleEngine/Private/ChaosModularVehicle/ModularVehicleSimulationCU.cpp b/Engine/Plugins/Experimental/ChaosModularVehicle/Source/ChaosModularVehicleEngine/Private/ChaosModularVehicle/ModularVehicleSimulationCU.cpp
--- a/Engine/Plugins/Experimental/ChaosModularVehicle/Source/ChaosModularVehicleEngine/Private/ChaosModularVehicle/ModularVehicleSimulationCU.cpp
+++ b/Engine/Plugins/Experimental/ChaosModularVehicle/Source/ChaosModularVehicleEngine/Private/ChaosModularVehicle/ModularVehicleSimulationCU.cpp
@@ -154,183 +154,184 @@ void FModularVehicleSimulationCU::PerformAdditionalSimWork(UWorld* InWorld, cons
 		const FClusterUnionIndex& CUI = Proxy->GetClusterUnionIndex();
 		if (FClusterUnion* ClusterUnion = ClusterUnionManager.FindClusterUnion(CUI))
 		{
-			FPBDRigidClusteredParticleHandle* ClusterHandle = ClusterUnion->InternalCluster;
-			TArray<FPBDRigidParticleHandle*> Particles = ClusterUnion->ChildParticles;
-			FRigidTransform3 ClusterWorldTM = FRigidTransform3(ClusterHandle->X(), ClusterHandle->R());
+			if (FPBDRigidClusteredParticleHandle* ClusterHandle = ClusterUnion->InternalCluster)
+			{
+				TArray<FPBDRigidParticleHandle*> Particles = ClusterUnion->ChildParticles;
+				FRigidTransform3 ClusterWorldTM = FRigidTransform3(ClusterHandle->X(), ClusterHandle->R());
 
-			//Chaos::FDebugDrawQueue::GetInstance().DrawDebugCoordinateSystem(ClusterHandle->X(), FRotator(ClusterHandle->R()), 200, false, -1.f, 1.f, 3.0f);
+				//Chaos::FDebugDrawQueue::GetInstance().DrawDebugCoordinateSystem(ClusterHandle->X(), FRotator(ClusterHandle->R()), 200, false, -1.f, 1.f, 3.0f);
 
-			const TArray<Chaos::FSimModuleTree::FSimModuleNode>& ModuleArray = SimModuleTree->GetSimulationModuleTree();
+				const TArray<Chaos::FSimModuleTree::FSimModuleNode>& ModuleArray = SimModuleTree->GetSimulationModuleTree();
 
-			for (const Chaos::FSimModuleTree::FSimModuleNode& Node : ModuleArray)
-			{
-				if (Node.IsValid() && Node.SimModule && Node.SimModule->IsEnabled())
+				for (const Chaos::FSimModuleTree::FSimModuleNode& Node : ModuleArray)
 				{
-					FRigidTransform3 Frame = FRigidTransform3::Identity;
-
-					FPBDRigidParticleHandle* Child = Node.SimModule->GetParticleFromUniqueIndex(Node.SimModule->GetParticleIndex().Idx, Particles);
-					if (Child == nullptr)
+					if (Node.IsValid() && Node.SimModule && Node.SimModule->IsEnabled())
 					{
-						continue;
-					}
+						FRigidTransform3 Frame = FRigidTransform3::Identity;
 
-					if (FPBDRigidClusteredParticleHandle* ClusterChild = Child->CastToClustered(); ClusterChild && ClusterChild->IsChildToParentLocked())
-					{
-						Frame = ClusterChild->ChildToParent();
-					}
-					else
-					{
-						Frame = ClusterChild->ChildToParent();
+						FPBDRigidParticleHandle* Child = Node.SimModule->GetParticleFromUniqueIndex(Node.SimModule->GetParticleIndex().Idx, Particles);
+						if (Child == nullptr)
+						{
+							continue;
+						}
 
-						const FRigidTransform3 ChildWorldTM(Child->X(), Child->R());
-						Frame = ChildWorldTM.GetRelativeTransform(ClusterWorldTM);
-					}
+						if (FPBDRigidClusteredParticleHandle* ClusterChild = Child->CastToClustered(); ClusterChild && ClusterChild->IsChildToParentLocked())
+						{
+							Frame = ClusterChild->ChildToParent();
+						}
+						else
+						{
+							Frame = ClusterChild->ChildToParent();
 
-					if (ClusterHandle)
-					{
-						AllInputs.VehicleWorldTransform = ClusterWorldTM;
+							const FRigidTransform3 ChildWorldTM(Child->X(), Child->R());
+							Frame = ChildWorldTM.GetRelativeTransform(ClusterWorldTM);
+						}
 
-						if (Node.SimModule->IsClustered() && Node.SimModule->IsBehaviourType(Chaos::eSimModuleTypeFlags::Raycast))
+						if (ClusterHandle)
 						{
-							Chaos::FSpringTrace OutTrace;
-							Chaos::FSuspensionSimModule* Suspension = static_cast<Chaos::FSuspensionSimModule*>(Node.SimModule);
+							AllInputs.VehicleWorldTransform = ClusterWorldTM;
 
-							// would be cleaner an faster to just store radius in suspension also
-							float WheelRadius = 0;
-							if (Suspension->GetWheelSimTreeIndex() != Chaos::ISimulationModuleBase::INVALID_IDX)
+							if (Node.SimModule->IsClustered() && Node.SimModule->IsBehaviourType(Chaos::eSimModuleTypeFlags::Raycast))
 							{
-								Chaos::FWheelSimModule* Wheel = static_cast<Chaos::FWheelSimModule*>(ModuleArray[Suspension->GetWheelSimTreeIndex()].SimModule);
-								if (Wheel)
+								Chaos::FSpringTrace OutTrace;
+								Chaos::FSuspensionSimModule* Suspension = static_cast<Chaos::FSuspensionSimModule*>(Node.SimModule);
+
+								// would be cleaner an faster to just store radius in suspension also
+								float WheelRadius = 0;
+								if (Suspension->GetWheelSimTreeIndex() != Chaos::ISimulationModuleBase::INVALID_IDX)
 								{
-									WheelRadius = Wheel->Setup().Radius;
+									Chaos::FWheelSimModule* Wheel = static_cast<Chaos::FWheelSimModule*>(ModuleArray[Suspension->GetWheelSimTreeIndex()].SimModule);
+									if (Wheel)
+									{
+										WheelRadius = Wheel->Setup().Radius;
+									}
 								}
-							}
 
-							Suspension->GetWorldRaycastLocation(ClusterWorldTM, WheelRadius, OutTrace);
+								Suspension->GetWorldRaycastLocation(ClusterWorldTM, WheelRadius, OutTrace);
 
-							FVector TraceStart = OutTrace.Start;
-							FVector TraceEnd = OutTrace.End;
+								FVector TraceStart = OutTrace.Start;
+								FVector TraceEnd = OutTrace.End;
 
-							const FCollisionQueryParams& TraceParams = InputData.PhysicsInputs.TraceParams;
-							FVector TraceVector(TraceStart - TraceEnd);
-							FVector TraceNormal = TraceVector.GetSafeNormal();
+								const FCollisionQueryParams& TraceParams = InputData.PhysicsInputs.TraceParams;
+								FVector TraceVector(TraceStart - TraceEnd);
+								FVector TraceNormal = TraceVector.GetSafeNormal();
 
-							FHitResult HitResult = FHitResult();
-							ECollisionChannel SpringCollisionChannel = ECollisionChannel::ECC_WorldDynamic;
-							const FCollisionResponseParams& ResponseParams = InputData.PhysicsInputs.TraceCollisionResponse;
-							if (InWorld)
-							{
-								InWorld->LineTraceSingleByChannel(HitResult, TraceStart, TraceEnd, SpringCollisionChannel, TraceParams, ResponseParams);
-							}
-
-							float Offset = Suspension->Setup().MaxLength;
-							if (HitResult.bBlockingHit && GModularVehicleDebugParams.SuspensionRaycastsEnabled)
-							{
-								Offset = HitResult.Distance - WheelRadius;
+								FHitResult HitResult = FHitResult();
+								ECollisionChannel SpringCollisionChannel = ECollisionChannel::ECC_WorldDynamic;
+								const FCollisionResponseParams& ResponseParams = InputData.PhysicsInputs.TraceCollisionResponse;
+								if (InWorld)
+								{
+									InWorld->LineTraceSingleByChannel(HitResult, TraceStart, TraceEnd, SpringCollisionChannel, TraceParams, ResponseParams);
+								}
 
-								if (Suspension->GetWheelSimTreeIndex() != Chaos::ISimulationModuleBase::INVALID_IDX)
+								float Offset = Suspension->Setup().MaxLength;
+								if (HitResult.bBlockingHit && GModularVehicleDebugParams.SuspensionRaycastsEnabled)
 								{
-									const Chaos::FSimModuleTree::FSimModuleNode& WheelNode = ModuleArray[Suspension->GetWheelSimTreeIndex()];
+									Offset = HitResult.Distance - WheelRadius;
 
-									Chaos::FWheelSimModule* Wheel = static_cast<Chaos::FWheelSimModule*>(WheelNode.SimModule);
-									if (Wheel && HitResult.PhysMaterial.IsValid())
+									if (Suspension->GetWheelSimTreeIndex() != Chaos::ISimulationModuleBase::INVALID_IDX)
 									{
-										if (GModularVehicleDebugParams.FrictionOverride > 0)
-										{
-											Wheel->SetSurfaceFriction(GModularVehicleDebugParams.FrictionOverride);
-										}
-										else
+										const Chaos::FSimModuleTree::FSimModuleNode& WheelNode = ModuleArray[Suspension->GetWheelSimTreeIndex()];
+
+										Chaos::FWheelSimModule* Wheel = static_cast<Chaos::FWheelSimModule*>(WheelNode.SimModule);
+										if (Wheel && HitResult.PhysMaterial.IsValid())
 										{
-											Wheel->SetSurfaceFriction(HitResult.PhysMaterial->Friction);
+											if (GModularVehicleDebugParams.FrictionOverride > 0)
+											{
+												Wheel->SetSurfaceFriction(GModularVehicleDebugParams.FrictionOverride);
+											}
+											else
+											{
+												Wheel->SetSurfaceFriction(HitResult.PhysMaterial->Friction);
+											}
 										}
 									}
-								}
 
 #if CHAOS_DEBUG_DRAW
-								if (GModularVehicleDebugParams.ShowSuspensionRaycasts)
-								{
-									Chaos::FDebugDrawQueue::GetInstance().DrawDebugSphere(HitResult.ImpactPoint, 3, 16, FColor::Red, false, -1.f, 0, 10.f);
-								}
+									if (GModularVehicleDebugParams.ShowSuspensionRaycasts)
+									{
+										Chaos::FDebugDrawQueue::GetInstance().DrawDebugSphere(HitResult.ImpactPoint, 3, 16, FColor::Red, false, -1.f, 0, 10.f);
+									}
 
-								if (Suspension->GetWheelSimTreeIndex() != Chaos::ISimulationModuleBase::INVALID_IDX)
-								{
-									Chaos::FWheelSimModule* Wheel = static_cast<Chaos::FWheelSimModule*>(ModuleArray[Suspension->GetWheelSimTreeIndex()].SimModule);
-									if (Wheel)
+									if (Suspension->GetWheelSimTreeIndex() != Chaos::ISimulationModuleBase::INVALID_IDX)
 									{
-										if (GModularVehicleDebugParams.ShowWheelData)
+										Chaos::FWheelSimModule* Wheel = static_cast<Chaos::FWheelSimModule*>(ModuleArray[Suspension->GetWheelSimTreeIndex()].SimModule);
+										if (Wheel)
 										{
-											FString TextOut = FString::Format(TEXT("{0}"), { Wheel->GetForceIntoSurface() });
-											FColor Col = FColor::White;
-											if (InWorld)
+											if (GModularVehicleDebugParams.ShowWheelData)
 											{
-												if (InWorld->GetNetMode() == ENetMode::NM_Client)
+												FString TextOut = FString::Format(TEXT("{0}"), { Wheel->GetForceIntoSurface() });
+												FColor Col = FColor::White;
+												if (InWorld)
 												{
-													Col = FColor::Blue;
-												}
-												else
-												{
-													Col = FColor::Red;
+													if (InWorld->GetNetMode() == ENetMode::NM_Client)
+													{
+														Col = FColor::Blue;
+													}
+													else
+													{
+														Col = FColor::Red;
+													}
 												}
+												Chaos::FDebugDrawQueue::GetInstance().DrawDebugString(HitResult.ImpactPoint + FVec3(0, 50, 50), TextOut, nullptr, Col, -1.f, true, 1.0f);
 											}
-											Chaos::FDebugDrawQueue::GetInstance().DrawDebugString(HitResult.ImpactPoint + FVec3(0, 50, 50), TextOut, nullptr, Col, -1.f, true, 1.0f);
 										}
 									}
-								}
 
 #endif
-							}
+								}
 
 #if CHAOS_DEBUG_DRAW
-							if (GModularVehicleDebugParams.ShowSuspensionRaycasts)
-							{
-								FColor DrawColor = FColor::Green;
-								DrawColor = (HitResult.bBlockingHit) ? FColor::Red : FColor::Green;
-								Chaos::FDebugDrawQueue::GetInstance().DrawDebugLine(TraceStart, TraceEnd, DrawColor, false, -1.f, 0, 2.f);
-								Chaos::FDebugDrawQueue::GetInstance().DrawDebugSphere(TraceStart, 3, 16, FColor::White, false, -1.f, 0, 10.f);
-								Chaos::FDebugDrawQueue::GetInstance().DrawDebugSphere(HitResult.ImpactPoint, 1, 16, FColor::Red, false, -1.f, 0, 10.f);
-								FString TextOut = FString::Format(TEXT("{0}"), { HitResult.Time});
-
-								FColor Col = FColor::White;
-								if (InWorld)
+								if (GModularVehicleDebugParams.ShowSuspensionRaycasts)
 								{
-									if (InWorld->GetNetMode() == ENetMode::NM_Client)
+									FColor DrawColor = FColor::Green;
+									DrawColor = (HitResult.bBlockingHit) ? FColor::Red : FColor::Green;
+									Chaos::FDebugDrawQueue::GetInstance().DrawDebugLine(TraceStart, TraceEnd, DrawColor, false, -1.f, 0, 2.f);
+									Chaos::FDebugDrawQueue::GetInstance().DrawDebugSphere(TraceStart, 3, 16, FColor::White, false, -1.f, 0, 10.f);
+									Chaos::FDebugDrawQueue::GetInstance().DrawDebugSphere(HitResult.ImpactPoint, 1, 16, FColor::Red, false, -1.f, 0, 10.f);
+									FString TextOut = FString::Format(TEXT("{0}"), { HitResult.Time });
+
+									FColor Col = FColor::White;
+									if (InWorld)
 									{
-										Col = FColor::Blue;
+										if (InWorld->GetNetMode() == ENetMode::NM_Client)
+										{
+											Col = FColor::Blue;
+										}
+										else
+										{
+											Col = FColor::Red;
+										}
 									}
-									else
+									Chaos::FDebugDrawQueue::GetInstance().DrawDebugString(HitResult.ImpactPoint + FVec3(0, 50, 50), TextOut, nullptr, Col, -1.f, true, 1.0f);
+								}
+
+								if (GModularVehicleDebugParams.ShowRaycastMaterial)
+								{
+									if (HitResult.PhysMaterial.IsValid())
 									{
-										Col = FColor::Red;
+										FDebugDrawQueue::GetInstance().DrawDebugString(HitResult.ImpactPoint, HitResult.PhysMaterial->GetName(), nullptr, FColor::White, -1.f, true, 1.0f);
 									}
 								}
-								Chaos::FDebugDrawQueue::GetInstance().DrawDebugString(HitResult.ImpactPoint + FVec3(0, 50, 50), TextOut, nullptr, Col, -1.f, true, 1.0f);
-							}
 
-							if (GModularVehicleDebugParams.ShowRaycastMaterial)
-							{
-								if (HitResult.PhysMaterial.IsValid())
+								if (GModularVehicleDebugParams.ShowWheelCollisionNormal)
 								{
-									FDebugDrawQueue::GetInstance().DrawDebugString(HitResult.ImpactPoint, HitResult.PhysMaterial->GetName(), nullptr, FColor::White, -1.f, true, 1.0f);
+									FVector Pt = HitResult.ImpactPoint;
+									FDebugDrawQueue::GetInstance().DrawDebugLine(Pt, Pt + HitResult.Normal * 20.0f, FColor::Yellow, false, 1.0f, 0, 1.0f);
+									FDebugDrawQueue::GetInstance().DrawDebugSphere(Pt, 5.0f, 4, FColor::White, false, 1.0f, 0, 1.0f);
 								}
-							}
 
-							if (GModularVehicleDebugParams.ShowWheelCollisionNormal)
-							{
-								FVector Pt = HitResult.ImpactPoint;
-								FDebugDrawQueue::GetInstance().DrawDebugLine(Pt, Pt + HitResult.Normal * 20.0f, FColor::Yellow, false, 1.0f, 0, 1.0f);
-								FDebugDrawQueue::GetInstance().DrawDebugSphere(Pt, 5.0f, 4, FColor::White, false, 1.0f, 0, 1.0f);
+#endif
+								Suspension->SetSpringLength(Offset, WheelRadius);
+								FVector Up = ClusterWorldTM.GetUnitAxis(EAxis::Z);
+								Suspension->SetTargetPoint(HitResult.ImpactPoint + Up * WheelRadius, HitResult.ImpactNormal, HitResult.bBlockingHit);
 							}
 
-#endif
-							Suspension->SetSpringLength(Offset, WheelRadius);
-							FVector Up = ClusterWorldTM.GetUnitAxis(EAxis::Z);
-							Suspension->SetTargetPoint(HitResult.ImpactPoint + Up * WheelRadius, HitResult.ImpactNormal, HitResult.bBlockingHit);
 						}
-
 					}
-				}
 
+				}
 			}
-
 		}
 
 	}
