Fixed CIS Issue 429340: missing headers following merge

diff --git a/Engine/Plugins/Experimental/ChaosModularVehicle/Source/ChaosModularVehicleEngine/Private/ChaosModularVehicle/ModularVehicleBaseComponent.cpp b/Engine/Plugins/Experimental/ChaosModularVehicle/Source/ChaosModularVehicleEngine/Private/ChaosModularVehicle/ModularVehicleBaseComponent.cpp
--- a/Engine/Plugins/Experimental/ChaosModularVehicle/Source/ChaosModularVehicleEngine/Private/ChaosModularVehicle/ModularVehicleBaseComponent.cpp
+++ b/Engine/Plugins/Experimental/ChaosModularVehicle/Source/ChaosModularVehicleEngine/Private/ChaosModularVehicle/ModularVehicleBaseComponent.cpp
@@ -6,6 +6,7 @@
 #include "ChaosModularVehicle/ModularVehicleDefaultAsyncInput.h"
 #include "Engine/Engine.h"
 #include "Engine/Canvas.h"
+#include "Engine/OverlapResult.h"
 #include "PhysicsEngine/PhysicsSettings.h"
 #include "PhysicsEngine/ClusterUnionComponent.h"
 #include "PhysicsEngine/PhysicsObjectExternalInterface.h"
diff --git a/Engine/Plugins/Experimental/ChaosVehiclesPlugin/Source/ChaosVehicles/Private/ChaosWheeledVehicleMovementComponent.cpp b/Engine/Plugins/Experimental/ChaosVehiclesPlugin/Source/ChaosVehicles/Private/ChaosWheeledVehicleMovementComponent.cpp
--- a/Engine/Plugins/Experimental/ChaosVehiclesPlugin/Source/ChaosVehicles/Private/ChaosWheeledVehicleMovementComponent.cpp
+++ b/Engine/Plugins/Experimental/ChaosVehiclesPlugin/Source/ChaosVehicles/Private/ChaosWheeledVehicleMovementComponent.cpp
@@ -9,6 +9,7 @@
 #include "DisplayDebugHelpers.h"
 #include "DisplayDebugHelpers.h"
 #include "Engine/Engine.h"
+#include "Engine/OverlapResult.h"
 #include "GameFramework/Pawn.h"
 #include "PhysicalMaterials/PhysicalMaterial.h"
 #include "PhysicsEngine/PhysicsAsset.h"
diff --git a/Engine/Plugins/Runtime/NetworkPredictionExtras/Source/NetworkPredictionExtras/Private/MockPhysicsGrenade.cpp b/Engine/Plugins/Runtime/NetworkPredictionExtras/Source/NetworkPredictionExtras/Private/MockPhysicsGrenade.cpp
--- a/Engine/Plugins/Runtime/NetworkPredictionExtras/Source/NetworkPredictionExtras/Private/MockPhysicsGrenade.cpp
+++ b/Engine/Plugins/Runtime/NetworkPredictionExtras/Source/NetworkPredictionExtras/Private/MockPhysicsGrenade.cpp
@@ -1,6 +1,7 @@
 // Copyright Epic Games, Inc. All Rights Reserved.
 
 #include "MockPhysicsGrenade.h"
+#include "Engine/OverlapResult.h"
 #include "NetworkPredictionModelDefRegistry.h"
 #include "NetworkPredictionProxyInit.h"
 #include "NetworkPredictionPhysics.h"
diff --git a/Engine/Plugins/Runtime/NetworkPredictionExtras/Source/NetworkPredictionExtras/Private/MockPhysicsSimulation.cpp b/Engine/Plugins/Runtime/NetworkPredictionExtras/Source/NetworkPredictionExtras/Private/MockPhysicsSimulation.cpp
--- a/Engine/Plugins/Runtime/NetworkPredictionExtras/Source/NetworkPredictionExtras/Private/MockPhysicsSimulation.cpp
+++ b/Engine/Plugins/Runtime/NetworkPredictionExtras/Source/NetworkPredictionExtras/Private/MockPhysicsSimulation.cpp
@@ -1,6 +1,7 @@
 // Copyright Epic Games, Inc. All Rights Reserved.
 
 #include "MockPhysicsSimulation.h"
+#include "Engine/OverlapResult.h"
 #include "Engine/World.h"
 #include "DrawDebugHelpers.h"
 #include "NetworkPredictionTickState.h"
