Fix transform setup error in modular vehicle

diff --git a/Engine/Plugins/Experimental/ChaosModularVehicle/Source/ChaosModularVehicleEngine/Private/ChaosModularVehicle/ModularVehicleBaseComponent.cpp b/Engine/Plugins/Experimental/ChaosModularVehicle/Source/ChaosModularVehicleEngine/Private/ChaosModularVehicle/ModularVehicleBaseComponent.cpp
--- a/Engine/Plugins/Experimental/ChaosModularVehicle/Source/ChaosModularVehicleEngine/Private/ChaosModularVehicle/ModularVehicleBaseComponent.cpp
+++ b/Engine/Plugins/Experimental/ChaosModularVehicle/Source/ChaosModularVehicleEngine/Private/ChaosModularVehicle/ModularVehicleBaseComponent.cpp
@@ -211,9 +211,8 @@ void UModularVehicleBaseComponent::CreateAssociatedSimComponents(const UPrimitiv
 
 	if (const UVehicleSimBaseComponent* Component = Cast<UVehicleSimBaseComponent>(AttachedComponent))
 	{
-		FTransform ComponentTransform = AttachedComponent->GetComponentTransform();
-		ComponentTransform.SetScale3D(FVector(1.0f, 1.0f, 1.0f));
-		ComponentTransform.SetTranslation(FVector::ZeroVector);
+		FTransform ClusterUnionComponentTransform = ClusterUnionComponent->GetComponentTransform();
+		FTransform ComponentTransform = AttachedComponent->GetComponentTransform().GetRelativeTransform(ClusterUnionComponentTransform);
 
 		int TreeIndex = INDEX_NONE;
 
