ModularVehicle: option to animate a skeletal mesh from modular physics output

diff --git a/Engine/Plugins/Experimental/ChaosModularVehicle/ChaosModularVehicle.uplugin b/Engine/Plugins/Experimental/ChaosModularVehicle/ChaosModularVehicle.uplugin
--- a/Engine/Plugins/Experimental/ChaosModularVehicle/ChaosModularVehicle.uplugin
+++ b/Engine/Plugins/Experimental/ChaosModularVehicle/ChaosModularVehicle.uplugin
@@ -24,7 +24,13 @@
 			"Name": "ChaosModularVehicleEngine",
 			"Type": "Runtime",
 			"LoadingPhase": "PreDefault"
+		},
+		{
+			"Name": "ChaosModularVehicleEditor",
+			"Type": "UncookedOnly",
+			"LoadingPhase": "PreDefault"
 		}
+
 	],
 	"Plugins": [
 		{
diff --git a/Engine/Plugins/Experimental/ChaosModularVehicle/Source/ChaosModularVehicleEditor/ChaosModularVehicleEditor.Build.cs b/Engine/Plugins/Experimental/ChaosModularVehicle/Source/ChaosModularVehicleEditor/ChaosModularVehicleEditor.Build.cs
new file mode 100644
--- /dev/null
+++ b/Engine/Plugins/Experimental/ChaosModularVehicle/Source/ChaosModularVehicleEditor/ChaosModularVehicleEditor.Build.cs
@@ -0,0 +1,33 @@
+// Copyright Epic Games, Inc. All Rights Reserved.
+
+namespace UnrealBuildTool.Rules
+{
+	public class ChaosModularVehicleEditor : ModuleRules
+	{
+		public ChaosModularVehicleEditor(ReadOnlyTargetRules Target) : base(Target)
+		{
+			SetupModulePhysicsSupport(Target);
+
+			PublicDependencyModuleNames.AddRange(
+				new string[]
+				{
+					"Core",
+					"CoreUObject",
+					"Chaos",
+					"Engine",
+					"EditorFramework",
+					"UnrealEd",
+					"PropertyEditor",
+					"AnimGraphRuntime",
+					"AnimGraph",
+					"BlueprintGraph",
+					"ToolMenus",
+					"AnimGraphRuntime",
+					"ChaosModularVehicleEngine"
+				}
+			);
+
+			PrivateDefinitions.Add("CHAOS_INCLUDE_LEVEL_1=1");
+		}
+	}
+}
diff --git a/Engine/Plugins/Experimental/ChaosModularVehicle/Source/ChaosModularVehicleEditor/Private/ChaosModularVehicle/AnimGraphNode_ModularVehicleController.cpp b/Engine/Plugins/Experimental/ChaosModularVehicle/Source/ChaosModularVehicleEditor/Private/ChaosModularVehicle/AnimGraphNode_ModularVehicleController.cpp
new file mode 100644
--- /dev/null
+++ b/Engine/Plugins/Experimental/ChaosModularVehicle/Source/ChaosModularVehicleEditor/Private/ChaosModularVehicle/AnimGraphNode_ModularVehicleController.cpp
@@ -0,0 +1,59 @@
+// Copyright Epic Games, Inc. All Rights Reserved.
+
+#include "ChaosModularVehicle/AnimGraphNode_ModularVehicleController.h"
+#include "Kismet2/BlueprintEditorUtils.h"
+#include "Kismet2/CompilerResultsLog.h"
+
+/////////////////////////////////////////////////////
+// UAnimGraphNode_ModularVehicleController
+
+#define LOCTEXT_NAMESPACE "ModularVehicle"
+
+UAnimGraphNode_ModularVehicleController::UAnimGraphNode_ModularVehicleController(const FObjectInitializer& ObjectInitializer)
+	: Super(ObjectInitializer)
+{
+}
+
+FText UAnimGraphNode_ModularVehicleController::GetControllerDescription() const
+{
+	return LOCTEXT("AnimGraphNode_ModularVehicleController", "Controller for ModularVehicle");
+}
+
+FText UAnimGraphNode_ModularVehicleController::GetTooltipText() const
+{
+	return LOCTEXT("AnimGraphNode_ModularVehicleController_Tooltip", "This alters the transform based on set up in Modular Vehicle. This only works when the owner is a modular vehicle.");
+}
+
+FText UAnimGraphNode_ModularVehicleController::GetNodeTitle(ENodeTitleType::Type TitleType) const
+{
+	FText NodeTitle;
+	if (TitleType == ENodeTitleType::ListView || TitleType == ENodeTitleType::MenuTitle)
+	{
+		NodeTitle = GetControllerDescription();
+	}
+	else
+	{
+		// we don't have any run-time information, so it's limited to print  
+		// anymore than what it is it would be nice to print more data such as 
+		// name of bones for wheels, but it's not available in Persona
+		NodeTitle = FText(LOCTEXT("AnimGraphNode_ModularVehicleController_Title", "Modular Vehicle Controller"));
+	}	
+	return NodeTitle;
+}
+
+void UAnimGraphNode_ModularVehicleController::ValidateAnimNodePostCompile(class FCompilerResultsLog& MessageLog, class UAnimBlueprintGeneratedClass* CompiledClass, int32 CompiledNodeIndex)
+{
+	// we only support vehicle anim instance
+	if ( CompiledClass->IsChildOf(UModularVehicleAnimationInstance::StaticClass())  == false )
+	{
+		MessageLog.Error(TEXT("@@ is only allowed in ModularVehicleAnimInstance. If this is for vehicle, please change parent to be ModularVehicleAnimInstance (Reparent Class)."), this);
+	}
+}
+
+bool UAnimGraphNode_ModularVehicleController::IsCompatibleWithGraph(const UEdGraph* TargetGraph) const
+{
+	UBlueprint* Blueprint = FBlueprintEditorUtils::FindBlueprintForGraph(TargetGraph);
+	return (Blueprint != nullptr) && Blueprint->ParentClass->IsChildOf<UModularVehicleAnimationInstance>() && Super::IsCompatibleWithGraph(TargetGraph);
+}
+
+#undef LOCTEXT_NAMESPACE
diff --git a/Engine/Plugins/Experimental/ChaosModularVehicle/Source/ChaosModularVehicleEditor/Private/ChaosModularVehicle/ChaosModularVehicleEditorPlugin.cpp b/Engine/Plugins/Experimental/ChaosModularVehicle/Source/ChaosModularVehicleEditor/Private/ChaosModularVehicle/ChaosModularVehicleEditorPlugin.cpp
new file mode 100644
--- /dev/null
+++ b/Engine/Plugins/Experimental/ChaosModularVehicle/Source/ChaosModularVehicleEditor/Private/ChaosModularVehicle/ChaosModularVehicleEditorPlugin.cpp
@@ -0,0 +1,24 @@
+// Copyright Epic Games, Inc. All Rights Reserved.
+
+#include "ChaosModularVehicle/ChaosModularVehicleEditorPlugin.h"
+
+#include "Modules/ModuleManager.h"
+#include "ChaosModularVehicle/ChaosSimModuleManager.h"
+
+
+class FChaosModularVehicleEditorPlugin : public IChaosModularVehicleEditorPlugin
+{
+  public:
+	/** IModuleInterface implementation */
+
+	virtual void StartupModule() override
+	{
+	}
+
+	virtual void ShutdownModule() override
+	{
+	}
+
+};
+
+IMPLEMENT_MODULE(FChaosModularVehicleEditorPlugin, ChaosModularVehicleEditor)
diff --git a/Engine/Plugins/Experimental/ChaosModularVehicle/Source/ChaosModularVehicleEditor/Public/ChaosModularVehicle/AnimGraphNode_ModularVehicleController.h b/Engine/Plugins/Experimental/ChaosModularVehicle/Source/ChaosModularVehicleEditor/Public/ChaosModularVehicle/AnimGraphNode_ModularVehicleController.h
new file mode 100644
--- /dev/null
+++ b/Engine/Plugins/Experimental/ChaosModularVehicle/Source/ChaosModularVehicleEditor/Public/ChaosModularVehicle/AnimGraphNode_ModularVehicleController.h
@@ -0,0 +1,35 @@
+// Copyright Epic Games, Inc. All Rights Reserved.
+
+#pragma once
+
+#include "CoreMinimal.h"
+#include "UObject/ObjectMacros.h"
+#include "AnimGraphNode_SkeletalControlBase.h"
+#include "ChaosModularVehicle/AnimNode_ModularVehicleController.h"
+#include "AnimGraphNode_ModularVehicleController.generated.h"
+
+class UEdGraph;
+
+UCLASS(MinimalAPI, meta=(Keywords = "Modify Modular Vehicle"))
+class UAnimGraphNode_ModularVehicleController : public UAnimGraphNode_SkeletalControlBase
+{
+	GENERATED_UCLASS_BODY()
+
+	UPROPERTY(EditAnywhere, Category=Settings)
+	FAnimNode_ModularVehicleController Node;
+
+public:
+	// UEdGraphNode interface
+	virtual FText GetNodeTitle(ENodeTitleType::Type TitleType) const override;
+	virtual FText GetTooltipText() const override;
+	// validate if this is within VehicleAnimInstance
+	virtual void ValidateAnimNodePostCompile(class FCompilerResultsLog& MessageLog, class UAnimBlueprintGeneratedClass* CompiledClass, int32 CompiledNodeIndex) override;
+	virtual bool IsCompatibleWithGraph(const UEdGraph* TargetGraph) const override;
+	// End of UEdGraphNode interface
+
+protected:
+	// UAnimGraphNode_SkeletalControlBase interface
+	virtual FText GetControllerDescription() const override;
+	virtual const FAnimNode_SkeletalControlBase* GetNode() const override { return &Node; }
+	// End of UAnimGraphNode_SkeletalControlBase interface
+};
diff --git a/Engine/Plugins/Experimental/ChaosModularVehicle/Source/ChaosModularVehicleEditor/Public/ChaosModularVehicle/ChaosModularVehicleEditorPlugin.h b/Engine/Plugins/Experimental/ChaosModularVehicle/Source/ChaosModularVehicleEditor/Public/ChaosModularVehicle/ChaosModularVehicleEditorPlugin.h
new file mode 100644
--- /dev/null
+++ b/Engine/Plugins/Experimental/ChaosModularVehicle/Source/ChaosModularVehicleEditor/Public/ChaosModularVehicle/ChaosModularVehicleEditorPlugin.h
@@ -0,0 +1,40 @@
+// Copyright Epic Games, Inc. All Rights Reserved.
+
+#pragma once
+
+#include "CoreMinimal.h"
+#include "Modules/ModuleInterface.h"
+#include "Modules/ModuleManager.h"
+#include "PhysicsInterfaceDeclaresCore.h"
+
+/**
+ * The public interface to this module
+ */
+class IChaosModularVehicleEditorPlugin : public IModuleInterface
+{
+
+public:
+
+	/**
+	 * Singleton-like access to this module's interface.  This is just for convenience!
+	 * Beware of calling this during the shutdown phase, though.  Your module might have been unloaded already.
+	 *
+	 * @return Returns singleton instance, loading the module on demand if needed
+	 */
+	static inline IChaosModularVehicleEditorPlugin& Get()
+	{
+		return FModuleManager::LoadModuleChecked< IChaosModularVehicleEditorPlugin >( "ChaosModularVehicleEditor" );
+	}
+
+	/**
+	 * Checks to see if this module is loaded and ready.  It is only valid to call Get() if IsAvailable() returns true.
+	 *
+	 * @return True if the module is loaded and ready to use
+	 */
+	static inline bool IsAvailable()
+	{
+		return FModuleManager::Get().IsModuleLoaded( "ChaosModularVehicleEditor" );
+	}
+
+private:
+};
diff --git a/Engine/Plugins/Experimental/ChaosModularVehicle/Source/ChaosModularVehicleEngine/ChaosModularVehicleEngine.Build.cs b/Engine/Plugins/Experimental/ChaosModularVehicle/Source/ChaosModularVehicleEngine/ChaosModularVehicleEngine.Build.cs
--- a/Engine/Plugins/Experimental/ChaosModularVehicle/Source/ChaosModularVehicleEngine/ChaosModularVehicleEngine.Build.cs
+++ b/Engine/Plugins/Experimental/ChaosModularVehicle/Source/ChaosModularVehicleEngine/ChaosModularVehicleEngine.Build.cs
@@ -17,6 +17,7 @@ namespace UnrealBuildTool.Rules
 					"Engine",
 					"RenderCore",
 					"RHI",
+					"AnimGraphRuntime",
 					"Renderer",
 					"ChaosVehiclesCore",
 					"ChaosModularVehicle",
diff --git a/Engine/Plugins/Experimental/ChaosModularVehicle/Source/ChaosModularVehicleEngine/Private/ChaosModularVehicle/AnimNode_ModularVehicleController.cpp b/Engine/Plugins/Experimental/ChaosModularVehicle/Source/ChaosModularVehicleEngine/Private/ChaosModularVehicle/AnimNode_ModularVehicleController.cpp
new file mode 100644
--- /dev/null
+++ b/Engine/Plugins/Experimental/ChaosModularVehicle/Source/ChaosModularVehicleEngine/Private/ChaosModularVehicle/AnimNode_ModularVehicleController.cpp
@@ -0,0 +1,157 @@
+// Copyright Epic Games, Inc. All Rights Reserved.
+
+#include "ChaosModularVehicle/AnimNode_ModularVehicleController.h"
+#include "AnimationRuntime.h"
+#include "Animation/AnimStats.h"
+#include "SimModule/SimulationModuleBase.h"
+
+
+FAnimNode_ModularVehicleController::FAnimNode_ModularVehicleController()
+{
+	AnimInstanceProxy = nullptr;
+}
+
+void FAnimNode_ModularVehicleController::GatherDebugData(FNodeDebugData& DebugData)
+{
+	FString DebugLine = DebugData.GetNodeName(this);
+
+	DebugLine += "(";
+	AddDebugNodeData(DebugLine);
+	DebugLine += ")";
+
+	DebugData.AddDebugItem(DebugLine);
+
+	const TArray<FModuleAnimationData>& AnimData = AnimInstanceProxy->GetModuleAnimData();
+	for (const FModuleLookupData& Module : Modules)
+	{
+		if (Module.BoneReference.BoneIndex != INDEX_NONE)
+		{
+			DebugLine = FString::Printf(TEXT(" [Module Index : %d] Bone: %s , Rotation Offset : %s, Location Offset : %s"),
+				Module.ModuleIndex, *Module.BoneReference.BoneName.ToString(), *AnimData[Module.ModuleIndex].RotOffset.ToString(), *AnimData[Module.ModuleIndex].LocOffset.ToString());
+		}
+		else
+		{
+			DebugLine = FString::Printf(TEXT(" [Module Index : %d] Bone: %s (invalid bone)"),
+				Module.ModuleIndex, *Module.BoneReference.BoneName.ToString());
+		}
+
+		DebugData.AddDebugItem(DebugLine);
+	}
+
+	ComponentPose.GatherDebugData(DebugData);
+}
+
+void FAnimNode_ModularVehicleController::EvaluateSkeletalControl_AnyThread(FComponentSpacePoseContext& Output, TArray<FBoneTransform>& OutBoneTransforms)
+{
+	check(OutBoneTransforms.Num() == 0);
+
+	ANIM_MT_SCOPE_CYCLE_COUNTER_VERBOSE(ModularVehicleController, !IsInGameThread());
+
+	const TArray<FModuleAnimationData>& ModuleAnimData = AnimInstanceProxy->GetModuleAnimData();
+
+	const FBoneContainer& BoneContainer = Output.Pose.GetPose().GetBoneContainer();
+	for(const FModuleLookupData& Module : Modules)
+	{
+		if (Module.BoneReference.IsValidToEvaluate(BoneContainer))
+		{
+			if (Module.ModuleIndex < ModuleAnimData.Num())
+			{
+				FCompactPoseBoneIndex ModuleSimBoneIndex = Module.BoneReference.GetCompactPoseIndex(BoneContainer);
+
+				// the way we apply transform is same as FMatrix or FTransform
+				// we apply scale first, and rotation, and translation
+				// if you'd like to translate first, you'll need two nodes that first node does translate and second nodes to rotate. 
+				FTransform NewBoneTM = Output.Pose.GetComponentSpaceTransform(ModuleSimBoneIndex);
+
+				FAnimationRuntime::ConvertCSTransformToBoneSpace(Output.AnimInstanceProxy->GetComponentTransform(), Output.Pose, NewBoneTM, ModuleSimBoneIndex, BCS_ComponentSpace);
+
+				if (ModuleAnimData[Module.ModuleIndex].Flags & Chaos::EAnimationFlags::AnimateRotation)
+				{
+					// Apply rotation offset
+					const FQuat BoneQuat(ModuleAnimData[Module.ModuleIndex].RotOffset);
+					NewBoneTM.SetRotation(BoneQuat * NewBoneTM.GetRotation());
+				}
+
+				if (ModuleAnimData[Module.ModuleIndex].Flags & Chaos::EAnimationFlags::AnimatePosition)
+				{
+					// Apply loc offset
+					NewBoneTM.AddToTranslation(ModuleAnimData[Module.ModuleIndex].LocOffset);
+				}
+
+				// Convert back to Component Space.
+				FAnimationRuntime::ConvertBoneSpaceTransformToCS(Output.AnimInstanceProxy->GetComponentTransform(), Output.Pose, NewBoneTM, ModuleSimBoneIndex, BCS_ComponentSpace);
+
+				// add back to it
+				OutBoneTransforms.Add(FBoneTransform(ModuleSimBoneIndex, NewBoneTM));
+			}
+		}
+	}
+
+#if ANIM_TRACE_ENABLED
+	for (const FModuleLookupData& Module : Modules)
+	{
+		if ((Module.BoneReference.BoneIndex != INDEX_NONE) && (Module.ModuleIndex < ModuleAnimData.Num()))
+		{
+			TRACE_ANIM_NODE_VALUE(Output, *FString::Printf(TEXT("Module %d Name"), Module.ModuleIndex), *Module.BoneReference.BoneName.ToString());
+			TRACE_ANIM_NODE_VALUE(Output, *FString::Printf(TEXT("Module %d Rotation Offset"), Module.ModuleIndex), ModuleAnimData[Module.ModuleIndex].RotOffset);
+			TRACE_ANIM_NODE_VALUE(Output, *FString::Printf(TEXT("Module %d Location Offset"), Module.ModuleIndex), ModuleAnimData[Module.ModuleIndex].LocOffset);
+		}
+		else
+		{
+			TRACE_ANIM_NODE_VALUE(Output, *FString::Printf(TEXT("Module %d Name"), Module.ModuleIndex), *FString::Printf(TEXT("%s (invalid)"), *Module.BoneReference.BoneName.ToString()));
+		}
+	}
+#endif
+}
+
+bool FAnimNode_ModularVehicleController::IsValidToEvaluate(const USkeleton* Skeleton, const FBoneContainer& RequiredBones) 
+{
+	if (AnimInstanceProxy)
+	{
+		// Note sure the best way to initilaize the animation since vehicle construction happens quite late on BeginPlay
+		const TArray<FModuleAnimationData>& ModuleAnimData = AnimInstanceProxy->GetModuleAnimData();
+		if (ModuleAnimData.Num() != Modules.Num())
+		{
+			InitializeBoneReferences(RequiredBones);
+		}
+	}
+
+	// if both bones are valid
+	for (const FModuleLookupData& Module : Modules)
+	{
+		// if one of them is valid
+		if (Module.BoneReference.IsValidToEvaluate(RequiredBones) == true)
+		{
+			return true;
+		}
+	}
+
+	return false;
+}
+
+void FAnimNode_ModularVehicleController::InitializeBoneReferences(const FBoneContainer& RequiredBones) 
+{
+	const TArray<FModuleAnimationData>& ModuleAnimData = AnimInstanceProxy->GetModuleAnimData();
+	const int32 NumModules = ModuleAnimData.Num();
+	Modules.Empty(NumModules);
+
+	for (int32 ModuleIndex = 0; ModuleIndex < NumModules; ++ModuleIndex)
+	{
+		FModuleLookupData* Module = new(Modules)FModuleLookupData();
+		Module->ModuleIndex = ModuleIndex;
+		Module->BoneReference.BoneName = ModuleAnimData[ModuleIndex].BoneName;
+		Module->BoneReference.Initialize(RequiredBones);
+	}
+
+	// sort by bone indices
+	Modules.Sort([](const FModuleLookupData& L, const FModuleLookupData& R) { return L.BoneReference.BoneIndex < R.BoneReference.BoneIndex; });
+}
+
+void FAnimNode_ModularVehicleController::Initialize_AnyThread(const FAnimationInitializeContext& Context)
+{
+	FAnimNode_SkeletalControlBase::Initialize_AnyThread(Context);
+
+	AnimInstanceProxy = (FModularVehicleAnimationInstanceProxy*)Context.AnimInstanceProxy;
+}
+
+
diff --git a/Engine/Plugins/Experimental/ChaosModularVehicle/Source/ChaosModularVehicleEngine/Private/ChaosModularVehicle/ModularVehicleAnimationInstance.cpp b/Engine/Plugins/Experimental/ChaosModularVehicle/Source/ChaosModularVehicleEngine/Private/ChaosModularVehicle/ModularVehicleAnimationInstance.cpp
new file mode 100644
--- /dev/null
+++ b/Engine/Plugins/Experimental/ChaosModularVehicle/Source/ChaosModularVehicleEngine/Private/ChaosModularVehicle/ModularVehicleAnimationInstance.cpp
@@ -0,0 +1,124 @@
+// Copyright Epic Games, Inc. All Rights Reserved.
+
+/*=============================================================================
+	UModularVehicleAnimationInstance.cpp: Single Node Tree Instance 
+	Only plays one animation at a time. 
+=============================================================================*/ 
+
+#include "ChaosModularVehicle/ModularVehicleAnimationInstance.h"
+#include "ChaosModularVehicle/ModularVehicleBaseComponent.h"
+#include "ChaosModularVehicle/ModularVehicleClusterPawn.h"
+#include "AnimationRuntime.h"
+
+/////////////////////////////////////////////////////
+// UModularVehicleAnimationInstance
+/////////////////////////////////////////////////////
+
+UModularVehicleAnimationInstance::UModularVehicleAnimationInstance(const FObjectInitializer& ObjectInitializer)
+	: Super(ObjectInitializer)
+{
+}
+
+class AModularVehicleClusterPawn* UModularVehicleAnimationInstance::GetVehicle()
+{
+	return Cast<AModularVehicleClusterPawn>(GetOwningActor());
+}
+
+void UModularVehicleAnimationInstance::NativeInitializeAnimation()
+{
+	// Find a wheeled movement component
+	if (AActor* Actor = GetOwningActor())
+	{
+		if (UModularVehicleBaseComponent* FoundModularVehicleComponent = Actor->FindComponentByClass<UModularVehicleBaseComponent>())
+		{
+			SetModularVehicleComponent(FoundModularVehicleComponent);
+		}
+	}
+}
+
+FAnimInstanceProxy* UModularVehicleAnimationInstance::CreateAnimInstanceProxy()
+{
+	return &AnimInstanceProxy;
+}
+
+void UModularVehicleAnimationInstance::DestroyAnimInstanceProxy(FAnimInstanceProxy* InProxy)
+{
+}
+
+/////////////////////////////////////////////////////
+//// PROXY ///
+/////////////////////////////////////////////////////
+
+void FModularVehicleAnimationInstanceProxy::SetModularVehicleComponent(const UModularVehicleBaseComponent* InWheeledVehicleComponent)
+{
+	const UModularVehicleBaseComponent* ModularVehicleComponent = InWheeledVehicleComponent;
+
+	//initialize data
+	// 
+	// take setup from modular vehicle and create anim instance data
+	const TArray<FModuleAnimationSetup>& ModuleAnimationSetups = ModularVehicleComponent->GetModuleAnimationSetups();
+
+	const int32 NumOfModules = ModuleAnimationSetups.Num();
+	ModuleInstances.Empty(NumOfModules);
+	if (NumOfModules > 0)
+	{
+		ModuleInstances.AddZeroed(NumOfModules);
+			
+		// now add Module data
+		for (int32 ModuleIndex = 0; ModuleIndex < ModuleInstances.Num(); ++ModuleIndex)
+		{
+			FModuleAnimationData& ModuleInstance = ModuleInstances[ModuleIndex];
+			const FModuleAnimationSetup& ModuleSetup = ModuleAnimationSetups[ModuleIndex];
+
+			// set data
+			ModuleInstance.BoneName = ModuleSetup.BoneName;
+			ModuleInstance.LocOffset = FVector::ZeroVector;
+			ModuleInstance.RotOffset = FRotator::ZeroRotator;
+		}
+	}
+}
+
+void FModularVehicleAnimationInstanceProxy::PreUpdate(UAnimInstance* InAnimInstance, float DeltaSeconds)
+{
+	Super::PreUpdate(InAnimInstance, DeltaSeconds);
+
+	if (const UModularVehicleAnimationInstance* VehicleAnimInstance = CastChecked<UModularVehicleAnimationInstance>(InAnimInstance))
+	{
+		// more modules have been added during runtime
+		if (const UModularVehicleBaseComponent* ModularVehicleComponent = VehicleAnimInstance->GetModularVehicleComponent())
+		{
+			if (ModuleInstances.Num() < ModularVehicleComponent->GetModuleAnimationSetups().Num())
+			{
+				int NumNew = ModularVehicleComponent->GetModuleAnimationSetups().Num() - ModuleInstances.Num();
+				int StartIdx = ModuleInstances.Num();
+				for (int I = 0; I < NumNew; I++)
+				{
+					FModuleAnimationData ModuleInstance;
+					ModuleInstance.BoneName = ModularVehicleComponent->GetModuleAnimationSetups()[StartIdx + I].BoneName;
+					ModuleInstance.LocOffset = FVector::ZeroVector;
+					ModuleInstance.RotOffset = FRotator::ZeroRotator;
+
+					ModuleInstances.Add(ModuleInstance);
+				}
+			}
+		}
+
+		if (const UModularVehicleBaseComponent* ModularVehicleComponent = VehicleAnimInstance->GetModularVehicleComponent())
+		{
+			for (int32 ModuleIndex = 0; ModuleIndex < ModuleInstances.Num(); ++ModuleIndex)
+			{
+				FModuleAnimationData& ModuleInstance = ModuleInstances[ModuleIndex];
+				if (ModularVehicleComponent->GetModuleAnimationSetups().IsValidIndex(ModuleIndex))
+				{ 
+					const FModuleAnimationSetup& ModuleAnim = ModularVehicleComponent->GetModuleAnimationSetups()[ModuleIndex];
+					{
+						ModuleInstance.LocOffset = ModuleAnim.LocOffset;
+						ModuleInstance.RotOffset = ModuleAnim.RotOffset;
+						ModuleInstance.Flags |= ModuleAnim.AnimFlags;
+					}
+				}
+			}
+		}
+	}
+}
+
diff --git a/Engine/Plugins/Experimental/ChaosModularVehicle/Source/ChaosModularVehicleEngine/Private/ChaosModularVehicle/ModularVehicleBaseComponent.cpp b/Engine/Plugins/Experimental/ChaosModularVehicle/Source/ChaosModularVehicleEngine/Private/ChaosModularVehicle/ModularVehicleBaseComponent.cpp
--- a/Engine/Plugins/Experimental/ChaosModularVehicle/Source/ChaosModularVehicleEngine/Private/ChaosModularVehicle/ModularVehicleBaseComponent.cpp
+++ b/Engine/Plugins/Experimental/ChaosModularVehicle/Source/ChaosModularVehicleEngine/Private/ChaosModularVehicle/ModularVehicleBaseComponent.cpp
@@ -4,6 +4,7 @@
 #include "ChaosModularVehicle/ChaosSimModuleManager.h"
 #include "ChaosModularVehicle/VehicleSimBaseComponent.h"
 #include "ChaosModularVehicle/ModularVehicleDefaultAsyncInput.h"
+#include "ChaosModularVehicle/ModularVehicleAnimationInstance.h"
 #include "Engine/Engine.h"
 #include "Engine/Canvas.h"
 #include "Engine/OverlapResult.h"
@@ -251,6 +252,32 @@ void UModularVehicleBaseComponent::CreateAssociatedSimComponents(const UPrimitiv
 		FTransform ClusterredTransform(FQuat::Identity, InitialTransform.GetLocation());
 		NewModule->SetClusteredTransform(ClusterredTransform);
 
+		if (Component->bAnimationEnabled && (Component->BoneName != NAME_None))
+		{
+			// if bone already exists then use that (seperate wheel and suspension modules can share same bone)
+
+			int FoundIndex = -1;
+			for (int I = 0; I < ModuleAnimationSetups.Num(); I++)
+			{
+				if (ModuleAnimationSetups[I].BoneName == Component->BoneName)
+				{
+					FoundIndex = I;
+					break;
+				}
+			}
+
+			if (FoundIndex != -1)
+			{
+				NewModule->SetAnimationData(Component->BoneName, Component->AnimationOffset, FoundIndex);
+			}
+			else
+			{
+				NewModule->SetAnimationData(Component->BoneName, Component->AnimationOffset, ModuleAnimationSetups.Num());
+				FModuleAnimationSetup AnimSetup(NewModule->GetBoneName());
+				ModuleAnimationSetups.Add(AnimSetup);
+			}
+		}
+
 		ParentIndex = TreeIndex;
 
 		if (bModularVehicle_SuspensionConstraint_Enabled)
@@ -672,14 +699,17 @@ void UModularVehicleBaseComponent::ParallelUpdate(float DeltaTime)
 		{
 			if (const FModularVehicleAsyncOutput* NextOutput = static_cast<FModularVehicleAsyncOutput*>(NextAsyncOutput))
 			{
-				PVehicleOutput->Clean();
+				UE_LOG(LogTemp, Warning, TEXT("Next output is valid"));
+			}
 
-				int NumItems = CurrentOutput->VehicleSimOutput.SimTreeOutputData.Num();
-				PVehicleOutput->SimTreeOutputData.Reserve(NumItems);
+			PVehicleOutput->Clean();
+			int NumItems = CurrentOutput->VehicleSimOutput.SimTreeOutputData.Num();
+			PVehicleOutput->SimTreeOutputData.Reserve(NumItems);
 
+			if (const FModularVehicleAsyncOutput* NextOutput = static_cast<FModularVehicleAsyncOutput*>(NextAsyncOutput))
+			{
 				for (int I = 0; I < NumItems; I++)
 				{
-					// #TODO: check we are lerping current/next of the same thing - the number can grow/shrink when construction/destruction happens
 					if (I < NextOutput->VehicleSimOutput.SimTreeOutputData.Num())
 					{
 						Chaos::FSimOutputData* CurrentSimData = CurrentOutput->VehicleSimOutput.SimTreeOutputData[I];
@@ -687,25 +717,54 @@ void UModularVehicleBaseComponent::ParallelUpdate(float DeltaTime)
 						PVehicleOutput->SimTreeOutputData.EmplaceAt(I, CurrentSimData->MakeNewData());
 						PVehicleOutput->SimTreeOutputData[I]->Lerp(*CurrentSimData, *NextSimData, OutputInterpAlpha);
 
-						// extract/cache some generally useful values as we go as trying to locate this data later requires a search
-						if (PVehicleOutput->SimTreeOutputData[I]->GetType() == Chaos::eSimType::Transmission)
+#if !(UE_BUILD_SHIPPING || UE_BUILD_TEST)
+						if (PVehicleOutput && !PVehicleOutput->SimTreeOutputData.IsEmpty() && PVehicleOutput->SimTreeOutputData[I])
 						{
-							// if there is more than one transmission then the last one will inform us of the current gear
-							CurrentGear = static_cast<Chaos::FTransmissionOutputData*>(PVehicleOutput->SimTreeOutputData[I])->CurrentGear;
+							PVehicleOutput->SimTreeOutputData[I]->DebugString = CurrentOutput->VehicleSimOutput.SimTreeOutputData[I]->DebugString;
 						}
-						else if (PVehicleOutput->SimTreeOutputData[I]->GetType() == Chaos::eSimType::Engine)
+#endif
+					}
+					else
+					{
+						Chaos::FSimOutputData* CurrentSimData = CurrentOutput->VehicleSimOutput.SimTreeOutputData[I];
+						Chaos::FSimOutputData* NextSimData = NextOutput->VehicleSimOutput.SimTreeOutputData[I];
+						PVehicleOutput->SimTreeOutputData.EmplaceAt(I, CurrentSimData->MakeNewData());
+						*PVehicleOutput->SimTreeOutputData[I] = *CurrentSimData;
+					}
+
+					// extract/cache some generally useful values as we go as trying to locate this data later requires a search
+					if (PVehicleOutput->SimTreeOutputData[I]->GetType() == Chaos::eSimType::Transmission)
+					{
+						// if there is more than one transmission then the last one will inform us of the current gear
+						CurrentGear = static_cast<Chaos::FTransmissionOutputData*>(PVehicleOutput->SimTreeOutputData[I])->CurrentGear;
+					}
+					else if (PVehicleOutput->SimTreeOutputData[I]->GetType() == Chaos::eSimType::Engine)
+					{
+						// if there is more than one engine then the last one will inform us of the engine RPM
+						Chaos::FEngineOutputData* Engine = static_cast<Chaos::FEngineOutputData*>(CurrentOutput->VehicleSimOutput.SimTreeOutputData[I]);
+
+						EngineRPM = Engine->RPM;
+						EngineTorque = Engine->Torque;
+					}
+
+					if (Chaos::FSimOutputData* ModuleOutput = PVehicleOutput->SimTreeOutputData[I])
+					{
+						if (ModuleOutput->AnimationSetupIndex >= 0)
 						{
-							// if there is more than one engine then the last one will inform us of the engine RPM
-							Chaos::FEngineOutputData* Engine = static_cast<Chaos::FEngineOutputData*>(PVehicleOutput->SimTreeOutputData[I]);
+							ModuleAnimationSetups[ModuleOutput->AnimationSetupIndex].AnimFlags |= ModuleOutput->AnimFlags;
 
-							EngineRPM = Engine->RPM;
-							EngineTorque = Engine->Torque;
-						}
+							if (ModuleOutput->AnimFlags & Chaos::EAnimationFlags::AnimateRotation)
+							{
+								ModuleAnimationSetups[ModuleOutput->AnimationSetupIndex].RotOffset = ModuleOutput->AnimationRotOffset;
+							}
 
-#if !(UE_BUILD_SHIPPING || UE_BUILD_TEST)
-						PVehicleOutput->SimTreeOutputData[I]->DebugString = NextOutput->VehicleSimOutput.SimTreeOutputData[I]->DebugString;
-#endif
+							if (ModuleOutput->AnimFlags & Chaos::EAnimationFlags::AnimatePosition)
+							{
+								ModuleAnimationSetups[ModuleOutput->AnimationSetupIndex].LocOffset = ModuleOutput->AnimationLocOffset;
+							}
+						}
 					}
+
 				}
 			}
 
diff --git a/Engine/Plugins/Experimental/ChaosModularVehicle/Source/ChaosModularVehicleEngine/Private/ChaosModularVehicle/VehicleSimBaseComponent.cpp b/Engine/Plugins/Experimental/ChaosModularVehicle/Source/ChaosModularVehicleEngine/Private/ChaosModularVehicle/VehicleSimBaseComponent.cpp
--- a/Engine/Plugins/Experimental/ChaosModularVehicle/Source/ChaosModularVehicleEngine/Private/ChaosModularVehicle/VehicleSimBaseComponent.cpp
+++ b/Engine/Plugins/Experimental/ChaosModularVehicle/Source/ChaosModularVehicleEngine/Private/ChaosModularVehicle/VehicleSimBaseComponent.cpp
@@ -6,9 +6,10 @@
 UVehicleSimBaseComponent::UVehicleSimBaseComponent(const FObjectInitializer& ObjectInitializer)
 	: Super(ObjectInitializer)
 {
-	TransformIndex = 0;
-	bRemoveFromClusterCollisionModel = false;
-	TreeIndex = -1;
+	BoneName = NAME_None;
+	AnimationOffset = FVector::ZeroVector;
 	bAnimationEnabled = false;
+	AnimationSetupIndex = -1;
+	TreeIndex = -1;
 }
 
diff --git a/Engine/Plugins/Experimental/ChaosModularVehicle/Source/ChaosModularVehicleEngine/Public/ChaosModularVehicle/AnimNode_ModularVehicleController.h b/Engine/Plugins/Experimental/ChaosModularVehicle/Source/ChaosModularVehicleEngine/Public/ChaosModularVehicle/AnimNode_ModularVehicleController.h
new file mode 100644
--- /dev/null
+++ b/Engine/Plugins/Experimental/ChaosModularVehicle/Source/ChaosModularVehicleEngine/Public/ChaosModularVehicle/AnimNode_ModularVehicleController.h
@@ -0,0 +1,46 @@
+// Copyright Epic Games, Inc. All Rights Reserved.
+
+#pragma once
+
+#include "CoreMinimal.h"
+#include "UObject/ObjectMacros.h"
+#include "BoneContainer.h"
+#include "BonePose.h"
+#include "ModularVehicleAnimationInstance.h"
+#include "BoneControllers/AnimNode_SkeletalControlBase.h"
+#include "AnimNode_ModularVehicleController.generated.h"
+
+/**
+ *	Simple controller that replaces or adds to the translation/rotation of a single bone.
+ */
+USTRUCT()
+struct CHAOSMODULARVEHICLEENGINE_API FAnimNode_ModularVehicleController : public FAnimNode_SkeletalControlBase
+{
+	GENERATED_USTRUCT_BODY()
+
+	FAnimNode_ModularVehicleController();
+
+	// FAnimNode_Base interface
+	virtual void GatherDebugData(FNodeDebugData& DebugData) override;
+	// End of FAnimNode_Base interface
+
+	// FAnimNode_SkeletalControlBase interface
+	virtual void EvaluateSkeletalControl_AnyThread(FComponentSpacePoseContext& Output, TArray<FBoneTransform>& OutBoneTransforms) override;
+	virtual bool IsValidToEvaluate(const USkeleton* Skeleton, const FBoneContainer& RequiredBones) override;
+	virtual void Initialize_AnyThread(const FAnimationInitializeContext& Context) override;
+	// End of FAnimNode_SkeletalControlBase interface
+
+private:
+	// FAnimNode_SkeletalControlBase interface
+	virtual void InitializeBoneReferences(const FBoneContainer& RequiredBones) override;
+	// End of FAnimNode_SkeletalControlBase interface
+
+	struct FModuleLookupData
+	{
+		int32 ModuleIndex;
+		FBoneReference BoneReference;
+	};
+
+	TArray<FModuleLookupData> Modules;
+	const FModularVehicleAnimationInstanceProxy* AnimInstanceProxy;
+};
diff --git a/Engine/Plugins/Experimental/ChaosModularVehicle/Source/ChaosModularVehicleEngine/Public/ChaosModularVehicle/ModularVehicleAnimationInstance.h b/Engine/Plugins/Experimental/ChaosModularVehicle/Source/ChaosModularVehicleEngine/Public/ChaosModularVehicle/ModularVehicleAnimationInstance.h
new file mode 100644
--- /dev/null
+++ b/Engine/Plugins/Experimental/ChaosModularVehicle/Source/ChaosModularVehicleEngine/Public/ChaosModularVehicle/ModularVehicleAnimationInstance.h
@@ -0,0 +1,91 @@
+// Copyright Epic Games, Inc. All Rights Reserved.
+
+#pragma once
+
+#include "CoreMinimal.h"
+#include "UObject/ObjectMacros.h"
+#include "Animation/AnimInstance.h"
+#include "Animation/AnimInstanceProxy.h"
+#include "ModularVehicleAnimationInstance.generated.h"
+
+class UModularVehicleBaseComponent;
+
+struct FModuleAnimationData
+{
+	FName BoneName;
+	FRotator RotOffset;
+	FVector LocOffset;
+	uint16 Flags;
+};
+
+/** Proxy override for this UAnimInstance-derived class */
+USTRUCT()
+	struct CHAOSMODULARVEHICLEENGINE_API FModularVehicleAnimationInstanceProxy : public FAnimInstanceProxy
+{
+	GENERATED_BODY()
+
+		FModularVehicleAnimationInstanceProxy()
+		: FAnimInstanceProxy()
+	{
+	}
+
+	FModularVehicleAnimationInstanceProxy(UAnimInstance* Instance)
+		: FAnimInstanceProxy(Instance)
+	{
+	}
+
+public:
+
+	void SetModularVehicleComponent(const UModularVehicleBaseComponent* InWheeledVehicleComponent);
+
+	/** FAnimInstanceProxy interface begin*/
+	virtual void PreUpdate(UAnimInstance* InAnimInstance, float DeltaSeconds) override;
+	/** FAnimInstanceProxy interface end*/
+
+	const TArray<FModuleAnimationData>& GetModuleAnimData() const
+	{
+		return ModuleInstances;
+	}
+
+private:
+	TArray<FModuleAnimationData> ModuleInstances;
+};
+
+UCLASS(transient)
+	class CHAOSMODULARVEHICLEENGINE_API UModularVehicleAnimationInstance : public UAnimInstance
+{
+	GENERATED_UCLASS_BODY()
+
+		/** Makes a montage jump to the end of a named section. */
+		UFUNCTION(BlueprintCallable, Category = "Animation")
+		class AModularVehicleClusterPawn* GetVehicle();
+
+public:
+	TArray<TArray<FModuleAnimationData>> ModuleData;
+
+public:
+	void SetModularVehicleComponent(const UModularVehicleBaseComponent* InWheeledVehicleComponent)
+	{
+		ModularVehicleComponent = InWheeledVehicleComponent;
+		AnimInstanceProxy.SetModularVehicleComponent(InWheeledVehicleComponent);
+	}
+
+	const UModularVehicleBaseComponent* GetModularVehicleComponent() const
+	{
+		return ModularVehicleComponent;
+	}
+
+private:
+	/** UAnimInstance interface begin*/
+	virtual void NativeInitializeAnimation() override;
+	virtual FAnimInstanceProxy* CreateAnimInstanceProxy() override;
+	virtual void DestroyAnimInstanceProxy(FAnimInstanceProxy* InProxy) override;
+	/** UAnimInstance interface end*/
+
+	FModularVehicleAnimationInstanceProxy AnimInstanceProxy;
+
+	UPROPERTY(transient)
+	TObjectPtr<const UModularVehicleBaseComponent> ModularVehicleComponent;
+};
+
+
diff --git a/Engine/Plugins/Experimental/ChaosModularVehicle/Source/ChaosModularVehicleEngine/Public/ChaosModularVehicle/ModularVehicleBaseComponent.h b/Engine/Plugins/Experimental/ChaosModularVehicle/Source/ChaosModularVehicleEngine/Public/ChaosModularVehicle/ModularVehicleBaseComponent.h
--- a/Engine/Plugins/Experimental/ChaosModularVehicle/Source/ChaosModularVehicleEngine/Public/ChaosModularVehicle/ModularVehicleBaseComponent.h
+++ b/Engine/Plugins/Experimental/ChaosModularVehicle/Source/ChaosModularVehicleEngine/Public/ChaosModularVehicle/ModularVehicleBaseComponent.h
@@ -151,6 +151,34 @@ struct CHAOSMODULARVEHICLEENGINE_API FConstructionData
 };
 
 
+USTRUCT()
+struct CHAOSMODULARVEHICLEENGINE_API FModuleAnimationSetup
+{
+	GENERATED_USTRUCT_BODY()
+
+	FModuleAnimationSetup(FName BoneNameIn)
+		: BoneName(BoneNameIn)
+		, RotOffset(FRotator::ZeroRotator)
+		, LocOffset(FVector::ZeroVector)
+		, AnimFlags(0)
+	{
+	}
+
+	FModuleAnimationSetup() 
+		: BoneName(NAME_None)
+		, RotOffset(FRotator::ZeroRotator)
+		, LocOffset(FVector::ZeroVector)
+		, AnimFlags(0)
+	{
+	}
+
+	FName BoneName;
+	FRotator RotOffset;
+	FVector LocOffset;
+	uint16 AnimFlags;
+};
+
+
 UCLASS(ClassGroup = (Physics), meta = (BlueprintSpawnableComponent), hidecategories = (PlanarMovement, "Components|Movement|Planar", Activation, "Components|Activation"))
 class CHAOSMODULARVEHICLEENGINE_API UModularVehicleBaseComponent : public UPawnMovementComponent
 {
@@ -337,6 +365,9 @@ public:
 		SuspensionTraceCollisionResponses.SetResponse(Channel, NewResponse);
 	}
 
+	TArray<FModuleAnimationSetup>& AccessModuleAnimationSetups() { return ModuleAnimationSetups; }
+	const TArray<FModuleAnimationSetup>& GetModuleAnimationSetups() const { return ModuleAnimationSetups; }
+
 protected:
 
 	void CreateVehicleSim();
@@ -344,6 +375,7 @@ protected:
 	void UpdatePhysicalProperties();
 	void AddOverlappingComponentsToCluster();
 	void AddGeometryCollectionsFromOwnedActor();
+	void SetupSkeletalAnimationStructure();
 
 	void ActionTreeUpdates(Chaos::FSimTreeUpdates* NextTreeUpdates);
 
@@ -552,5 +584,8 @@ private:
 	int32 ClusteringCount = 0;
 
 	bool bIsLocallyControlled;
+
+	TArray<FModuleAnimationSetup> ModuleAnimationSetups;
+
 };
 
diff --git a/Engine/Plugins/Experimental/ChaosModularVehicle/Source/ChaosModularVehicleEngine/Public/ChaosModularVehicle/VehicleSimBaseComponent.h b/Engine/Plugins/Experimental/ChaosModularVehicle/Source/ChaosModularVehicleEngine/Public/ChaosModularVehicle/VehicleSimBaseComponent.h
--- a/Engine/Plugins/Experimental/ChaosModularVehicle/Source/ChaosModularVehicleEngine/Public/ChaosModularVehicle/VehicleSimBaseComponent.h
+++ b/Engine/Plugins/Experimental/ChaosModularVehicle/Source/ChaosModularVehicleEngine/Public/ChaosModularVehicle/VehicleSimBaseComponent.h
@@ -50,14 +50,15 @@ public:
 	virtual Chaos::ISimulationModuleBase* CreateNewCoreModule() const { return nullptr; }
 
 	UPROPERTY(EditAnywhere, BlueprintReadWrite, Category = ModularVehicle)
-	int TransformIndex;
+	FName BoneName;
 
 	UPROPERTY(EditAnywhere, BlueprintReadWrite, Category = ModularVehicle)
-	bool bRemoveFromClusterCollisionModel;
+	FVector AnimationOffset;
 
 	UPROPERTY(EditAnywhere, BlueprintReadWrite, Category = ModularVehicle)
 	bool bAnimationEnabled;
 
+	int AnimationSetupIndex;
 	int TreeIndex; // helper - since Component->GetAttachChildren doesn't contain any data
 };
 
diff --git a/Engine/Source/Runtime/Experimental/ChaosVehicles/ChaosVehiclesCore/Private/SimModule/SimulationModuleBase.cpp b/Engine/Source/Runtime/Experimental/ChaosVehicles/ChaosVehiclesCore/Private/SimModule/SimulationModuleBase.cpp
--- a/Engine/Source/Runtime/Experimental/ChaosVehicles/ChaosVehiclesCore/Private/SimModule/SimulationModuleBase.cpp
+++ b/Engine/Source/Runtime/Experimental/ChaosVehicles/ChaosVehiclesCore/Private/SimModule/SimulationModuleBase.cpp
@@ -93,7 +93,14 @@ FPBDRigidParticleHandle* ISimulationModuleBase::GetParticleFromUniqueIndex(int32
 	return nullptr;
 }
 
-bool ISimulationModuleBase::GetDebugString(FString& StringOut) const 
+void ISimulationModuleBase::SetAnimationData(const FName& BoneNameIn, const FVector& AnimationOffsetIn, int AnimationSetupIndexIn)
+{
+	BoneName = BoneNameIn;
+	AnimationOffset = AnimationOffsetIn;
+	AnimationSetupIndex = AnimationSetupIndexIn;
+}
+
+bool ISimulationModuleBase::GetDebugString(FString& StringOut) const
 {
 	StringOut += FString::Format(TEXT("{0}: TreeIndex {1}, Enabled {2}, InCluster {3}, TFormIdx {4}, ")
 		, { GetDebugName(), GetTreeIndex(), IsEnabled(), IsClustered(), GetTransformIndex() });
@@ -116,6 +123,8 @@ const FTransform& ISimulationModuleBase::GetParentRelativeTransform() const
 
 void FSimOutputData::FillOutputState(const ISimulationModuleBase* SimModule)
 {
+	AnimationSetupIndex = SimModule->AnimationSetupIndex;
+
 #if !(UE_BUILD_SHIPPING || UE_BUILD_TEST)
 	if (SimModule)
 	{	
@@ -125,5 +134,20 @@ void FSimOutputData::FillOutputState(const ISimulationModuleBase* SimModule)
 #endif
 }
 
+void FSimOutputData::Lerp(const FSimOutputData& InCurrent, const FSimOutputData& InNext, float Alpha)
+{
+	AnimationSetupIndex = InNext.AnimationSetupIndex;
+	AnimFlags = InNext.AnimFlags;
+	if (AnimFlags & EAnimationFlags::AnimatePosition)
+	{
+		AnimationLocOffset = FMath::Lerp(InCurrent.AnimationLocOffset, InNext.AnimationLocOffset, Alpha);
+	}
+
+	if (AnimFlags & EAnimationFlags::AnimateRotation)
+	{
+		AnimationRotOffset = FMath::Lerp(InCurrent.AnimationRotOffset, InNext.AnimationRotOffset, Alpha);
+	}
+}
+
 
 } //namespace Chaos
\ No newline at end of file
diff --git a/Engine/Source/Runtime/Experimental/ChaosVehicles/ChaosVehiclesCore/Private/SimModule/SuspensionModule.cpp b/Engine/Source/Runtime/Experimental/ChaosVehicles/ChaosVehiclesCore/Private/SimModule/SuspensionModule.cpp
--- a/Engine/Source/Runtime/Experimental/ChaosVehicles/ChaosVehiclesCore/Private/SimModule/SuspensionModule.cpp
+++ b/Engine/Source/Runtime/Experimental/ChaosVehicles/ChaosVehiclesCore/Private/SimModule/SuspensionModule.cpp
@@ -9,7 +9,7 @@
 #include "VehicleUtility.h"
 
 #if VEHICLE_DEBUGGING_ENABLED
-UE_DISABLE_OPTIMIZATION
+UE_DISABLE_OPTIMIZATION_SHIP
 #endif
 
 namespace Chaos
@@ -178,12 +178,19 @@ namespace Chaos
 		if (const FSuspensionSimModule* Sim = static_cast<const FSuspensionSimModule*>(SimModule))
 		{
 			SpringDisplacement = Sim->SpringDisplacement;
+			SpringDisplacementVector = -Sim->Setup().SuspensionAxis * Sim->SpringDisplacement + Sim->GetAnimationOffset();
 			SpringSpeed = Sim->SpringSpeed;
+
+			AnimFlags = EAnimationFlags::AnimatePosition;
+			AnimationLocOffset = SpringDisplacementVector;
+
 		}
 	}
 
 	void FSuspensionOutputData::Lerp(const FSimOutputData& InCurrent, const FSimOutputData& InNext, float Alpha)
 	{
+		FSimOutputData::Lerp(InCurrent, InNext, Alpha);
+
 		const FSuspensionOutputData& Current = static_cast<const FSuspensionOutputData&>(InCurrent);
 		const FSuspensionOutputData& Next = static_cast<const FSuspensionOutputData&>(InNext);
 
@@ -201,5 +208,5 @@ namespace Chaos
 } // namespace Chaos
 
 #if VEHICLE_DEBUGGING_ENABLED
-UE_ENABLE_OPTIMIZATION
+UE_ENABLE_OPTIMIZATION_SHIP
 #endif
diff --git a/Engine/Source/Runtime/Experimental/ChaosVehicles/ChaosVehiclesCore/Private/SimModule/WheelModule.cpp b/Engine/Source/Runtime/Experimental/ChaosVehicles/ChaosVehiclesCore/Private/SimModule/WheelModule.cpp
--- a/Engine/Source/Runtime/Experimental/ChaosVehicles/ChaosVehiclesCore/Private/SimModule/WheelModule.cpp
+++ b/Engine/Source/Runtime/Experimental/ChaosVehicles/ChaosVehiclesCore/Private/SimModule/WheelModule.cpp
@@ -5,7 +5,7 @@
 #include "VehicleUtility.h"
 
 #if VEHICLE_DEBUGGING_ENABLED
-UE_DISABLE_OPTIMIZATION
+UE_DISABLE_OPTIMIZATION_SHIP
 #endif
 
 namespace Chaos
@@ -249,11 +249,19 @@ namespace Chaos
 			ForceIntoSurface = Sim->ForceIntoSurface;
 			SlipAngle = Sim->SlipAngle;
 			RPM = Sim->GetRPM();
+			AngularPositionDegrees = -Sim->AngularPosition * 180.0f / PI;
+			SteeringAngleDegrees = Sim->SteerAngleDegrees;
+
+			AnimFlags = EAnimationFlags::AnimateRotation;
+			AnimationRotOffset.Pitch = AngularPositionDegrees;
+			AnimationRotOffset.Yaw = SteeringAngleDegrees;
 		}
 	}
 
 	void FWheelOutputData::Lerp(const FSimOutputData& InCurrent, const FSimOutputData& InNext, float Alpha)
 	{
+		FSimOutputData::Lerp(InCurrent, InNext, Alpha);
+
 		const FWheelOutputData& Current = static_cast<const FWheelOutputData&>(InCurrent);
 		const FWheelOutputData& Next = static_cast<const FWheelOutputData&>(InNext);
 
@@ -275,5 +283,5 @@ namespace Chaos
 
 
 #if VEHICLE_DEBUGGING_ENABLED
-UE_ENABLE_OPTIMIZATION
+UE_ENABLE_OPTIMIZATION_SHIP
 #endif
diff --git a/Engine/Source/Runtime/Experimental/ChaosVehicles/ChaosVehiclesCore/Private/TransmissionSystem.cpp b/Engine/Source/Runtime/Experimental/ChaosVehicles/ChaosVehiclesCore/Private/TransmissionSystem.cpp
--- a/Engine/Source/Runtime/Experimental/ChaosVehicles/ChaosVehiclesCore/Private/TransmissionSystem.cpp
+++ b/Engine/Source/Runtime/Experimental/ChaosVehicles/ChaosVehiclesCore/Private/TransmissionSystem.cpp
@@ -3,7 +3,7 @@
 #include "TransmissionSystem.h"
 
 #if VEHICLE_DEBUGGING_ENABLED
-UE_DISABLE_OPTIMIZATION
+UE_DISABLE_OPTIMIZATION_SHIP
 #endif
 
 namespace Chaos
@@ -102,5 +102,5 @@ namespace Chaos
 } // namespace Chaos
 
 #if VEHICLE_DEBUGGING_ENABLED
-UE_ENABLE_OPTIMIZATION
+UE_ENABLE_OPTIMIZATION_SHIP
 #endif
diff --git a/Engine/Source/Runtime/Experimental/ChaosVehicles/ChaosVehiclesCore/Public/SimModule/SimulationModuleBase.h b/Engine/Source/Runtime/Experimental/ChaosVehicles/ChaosVehiclesCore/Public/SimModule/SimulationModuleBase.h
--- a/Engine/Source/Runtime/Experimental/ChaosVehicles/ChaosVehiclesCore/Public/SimModule/SimulationModuleBase.h
+++ b/Engine/Source/Runtime/Experimental/ChaosVehicles/ChaosVehiclesCore/Public/SimModule/SimulationModuleBase.h
@@ -162,11 +162,15 @@ namespace Chaos
 	 */
 	class CHAOSVEHICLESCORE_API ISimulationModuleBase
 	{
+	friend FSimOutputData;
+
 	public:
 		const static int INVALID_IDX = -1;
 
 		ISimulationModuleBase()
 			: SimModuleTree(nullptr)
+			, BoneName(NAME_None)
+			, AnimationSetupIndex(-1)
 			, SimTreeIndex(INVALID_IDX)
 			, StateFlags(Enabled)
 			, TransformIndex(INVALID_IDX)
@@ -221,6 +225,11 @@ namespace Chaos
 		void SetAnimationEnabled(bool bInEnabled) { bAnimationEnabled = bInEnabled; }
 		bool IsAnimationEnabled() { return bAnimationEnabled; }
 
+		void SetAnimationData(const FName& BoneNameIn, const FVector& AnimationOffsetIn, int AnimationSetupIndexIn);
+		const FVector& GetAnimationOffset() const { return AnimationOffset; }
+		const FName& GetBoneName() const { return BoneName; }
+		int GetAnimationSetupIndex() { return AnimationSetupIndex; }
+
 		/**
 		 * Option to draw debug for this module requires CVar p.Chaos.DebugDraw.Enabled 1
 		 */
@@ -332,6 +341,8 @@ namespace Chaos
 	protected:
 
 		FSimModuleTree* SimModuleTree;	// A pointer back to the simulation tree where we are stored
+		FName BoneName;
+		int AnimationSetupIndex;
 		int SimTreeIndex;	// Index of this SimModule in the FSimModuleTree
 		eSimModuleState StateFlags;	// TODO: make this more like flags
 		int TransformIndex; // Index of this Sim Module's node in Geometry Collection Transform array
@@ -347,6 +358,8 @@ namespace Chaos
 		FVector LocalAngularVelocity;
 		bool bClustered;
 		bool bAnimationEnabled;
+		FVector AnimationOffset; 
+		FVector AnimationRotation;
 
 		// for headless chaos testing
 		FVector AppliedForce;
@@ -385,6 +398,13 @@ namespace Chaos
 
 	using FModuleNetDataArray = TArray<TSharedPtr<FModuleNetData>>;
 
+	namespace EAnimationFlags
+	{
+		static uint16 AnimateNone = 0x00000000;
+		static uint16 AnimatePosition = 0x00000001;
+		static uint16 AnimateRotation = 0x00000002;
+	}
+
 	struct CHAOSVEHICLESCORE_API FSimOutputData
 	{
 		FSimOutputData() = default;
@@ -394,9 +414,13 @@ namespace Chaos
 		virtual bool IsEnabled() { return bEnabled; }
 		virtual FSimOutputData* MakeNewData() = 0;
 		virtual void FillOutputState(const ISimulationModuleBase* SimModule);
-		virtual void Lerp(const FSimOutputData& InCurrent, const FSimOutputData& InNext, float Alpha) = 0;
+		virtual void Lerp(const FSimOutputData& InCurrent, const FSimOutputData& InNext, float Alpha);
 
 		bool bEnabled = true;
+		int AnimationSetupIndex;
+		uint16 AnimFlags;
+		FVector AnimationLocOffset;
+		FRotator AnimationRotOffset;
 
 #if !(UE_BUILD_SHIPPING || UE_BUILD_TEST)
 		virtual FString ToString() { return FString(); }
diff --git a/Engine/Source/Runtime/Experimental/ChaosVehicles/ChaosVehiclesCore/Public/SimModule/SuspensionModule.h b/Engine/Source/Runtime/Experimental/ChaosVehicles/ChaosVehiclesCore/Public/SimModule/SuspensionModule.h
--- a/Engine/Source/Runtime/Experimental/ChaosVehicles/ChaosVehiclesCore/Public/SimModule/SuspensionModule.h
+++ b/Engine/Source/Runtime/Experimental/ChaosVehicles/ChaosVehiclesCore/Public/SimModule/SuspensionModule.h
@@ -55,6 +55,7 @@ namespace Chaos
 #endif
 
 		float SpringDisplacement;
+		FVector SpringDisplacementVector;
 		float SpringSpeed;
 	};
 
diff --git a/Engine/Source/Runtime/Experimental/ChaosVehicles/ChaosVehiclesCore/Public/SimModule/WheelModule.h b/Engine/Source/Runtime/Experimental/ChaosVehicles/ChaosVehiclesCore/Public/SimModule/WheelModule.h
--- a/Engine/Source/Runtime/Experimental/ChaosVehicles/ChaosVehiclesCore/Public/SimModule/WheelModule.h
+++ b/Engine/Source/Runtime/Experimental/ChaosVehicles/ChaosVehiclesCore/Public/SimModule/WheelModule.h
@@ -57,6 +57,8 @@ namespace Chaos
 		float ForceIntoSurface;
 		float SlipAngle;
 		float RPM;
+		float AngularPositionDegrees;
+		float SteeringAngleDegrees;
 
 		//HitLocation
 		//PhysMaterial
diff --git a/Engine/Source/Runtime/Experimental/ChaosVehicles/ChaosVehiclesCore/Public/TransmissionSystem.h b/Engine/Source/Runtime/Experimental/ChaosVehicles/ChaosVehiclesCore/Public/TransmissionSystem.h
--- a/Engine/Source/Runtime/Experimental/ChaosVehicles/ChaosVehiclesCore/Public/TransmissionSystem.h
+++ b/Engine/Source/Runtime/Experimental/ChaosVehicles/ChaosVehiclesCore/Public/TransmissionSystem.h
@@ -8,7 +8,7 @@
 #include "VehicleUtility.h"
 
 #if VEHICLE_DEBUGGING_ENABLED
-UE_DISABLE_OPTIMIZATION
+UE_DISABLE_OPTIMIZATION_SHIP
 #endif
 
 /**
@@ -237,6 +237,6 @@ namespace Chaos
 } // namespace Chaos
 
 #if VEHICLE_DEBUGGING_ENABLED
-UE_ENABLE_OPTIMIZATION
+UE_ENABLE_OPTIMIZATION_SHIP
 #endif
 
