Removing some logging that was left in

diff --git a/Engine/Plugins/Experimental/ChaosModularVehicle/Source/ChaosModularVehicleEngine/Private/ChaosModularVehicle/ModularVehicleBaseComponent.cpp b/Engine/Plugins/Experimental/ChaosModularVehicle/Source/ChaosModularVehicleEngine/Private/ChaosModularVehicle/ModularVehicleBaseComponent.cpp
--- a/Engine/Plugins/Experimental/ChaosModularVehicle/Source/ChaosModularVehicleEngine/Private/ChaosModularVehicle/ModularVehicleBaseComponent.cpp
+++ b/Engine/Plugins/Experimental/ChaosModularVehicle/Source/ChaosModularVehicleEngine/Private/ChaosModularVehicle/ModularVehicleBaseComponent.cpp
@@ -696,11 +696,6 @@ void UModularVehicleBaseComponent::ParallelUpdate(float DeltaTime)
 	{
 		if (CurrentOutput->bValid && PVehicleOutput)
 		{
-			if (const FModularVehicleAsyncOutput* NextOutput = static_cast<FModularVehicleAsyncOutput*>(NextAsyncOutput))
-			{
-				UE_LOG(LogTemp, Warning, TEXT("Next output is valid"));
-			}
-
 			PVehicleOutput->Clean();
 			int NumItems = CurrentOutput->VehicleSimOutput.SimTreeOutputData.Num();
 			PVehicleOutput->SimTreeOutputData.Reserve(NumItems);
