Core vehicle changes to make code more extensible wrt creating new custom modules

diff --git a/Engine/Plugins/Experimental/ChaosModularVehicle/Source/ChaosModularVehicleEngine/Private/ChaosModularVehicle/ChaosSimModuleManagerAsyncCallback.cpp b/Engine/Plugins/Experimental/ChaosModularVehicle/Source/ChaosModularVehicleEngine/Private/ChaosModularVehicle/ChaosSimModuleManagerAsyncCallback.cpp
--- a/Engine/Plugins/Experimental/ChaosModularVehicle/Source/ChaosModularVehicleEngine/Private/ChaosModularVehicle/ChaosSimModuleManagerAsyncCallback.cpp
+++ b/Engine/Plugins/Experimental/ChaosModularVehicle/Source/ChaosModularVehicleEngine/Private/ChaosModularVehicle/ChaosSimModuleManagerAsyncCallback.cpp
@@ -8,9 +8,11 @@
 #include "Chaos/ParticleHandleFwd.h"
 #include "PhysicsProxy/GeometryCollectionPhysicsProxy.h"
 #include "PhysicsProxy/ClusterUnionPhysicsProxy.h"
+#include "SimModule/ModuleFactoryRegister.h"
 
 FSimModuleDebugParams GSimModuleDebugParams;
 
+
 DECLARE_CYCLE_STAT(TEXT("AsyncCallback:OnPreSimulate_Internal"), STAT_AsyncCallback_OnPreSimulate, STATGROUP_ChaosSimModuleManager);
 
 FName FChaosSimModuleManagerAsyncCallback::GetFNameForStatId() const
@@ -277,63 +279,16 @@ bool FNetworkModularVehicleStates::NetSerialize(FArchive& Ar, class UPackageMap*
 				if (ModuleData.Num() != NumNetModules)
 				{
 					ModuleData.Reserve(NumNetModules);
-					switch (ModuleType)
-					{
-					case Chaos::eSimType::Suspension:
-					{
-						ModuleData.Emplace(MakeShared<Chaos::FSuspensionSimModuleDatas>(SimArrayIndex
-#if !(UE_BUILD_SHIPPING || UE_BUILD_TEST)
-							, FString()
-#endif
-						));
-					}
-					break;
-
-					case Chaos::eSimType::Transmission:
-					{
-						ModuleData.Emplace(MakeShared<Chaos::FTransmissionSimModuleDatas>(SimArrayIndex
-#if !(UE_BUILD_SHIPPING || UE_BUILD_TEST)
-							, FString()
-#endif
-						));
-					}
-					break;
 
-					case Chaos::eSimType::Engine:
+					if (TSharedPtr<Chaos::FModuleNetData> Data = Chaos::FModuleFactoryRegister::Get().GenerateNetData(ModuleType, SimArrayIndex))
 					{
-						ModuleData.Emplace(MakeShared<Chaos::FEngineSimModuleDatas>(SimArrayIndex
-#if !(UE_BUILD_SHIPPING || UE_BUILD_TEST)
-							, FString()
-#endif
-						));
+						ModuleData.Emplace(Data);
 					}
-					break;
-
-					case Chaos::eSimType::Clutch:
-					{
-						ModuleData.Emplace(MakeShared<Chaos::FClutchSimModuleDatas>(SimArrayIndex
-#if !(UE_BUILD_SHIPPING || UE_BUILD_TEST)
-							, FString()
-#endif
-						));
-					}
-					break;
-
-					case Chaos::eSimType::Wheel:
+					else
 					{
-						ModuleData.Emplace(MakeShared<Chaos::FWheelSimModuleDatas>(SimArrayIndex
-#if !(UE_BUILD_SHIPPING || UE_BUILD_TEST)
-							, FString()
-#endif
-						));
+						checkf(false, TEXT("Module net data factory has not been registered, use FModuleFactoryRegister::Get().RegisterFactory()"));
 					}
-					break;
 
-					default:
-					{
-						checkf(false, TEXT("Unhandled NetModuleType case"));
-					}
-					}
 				}
 			}
 		}
diff --git a/Engine/Plugins/Experimental/ChaosModularVehicle/Source/ChaosModularVehicleEngine/Private/ChaosModularVehicle/ModularVehicleBuilder.cpp b/Engine/Plugins/Experimental/ChaosModularVehicle/Source/ChaosModularVehicleEngine/Private/ChaosModularVehicle/ModularVehicleBuilder.cpp
--- a/Engine/Plugins/Experimental/ChaosModularVehicle/Source/ChaosModularVehicleEngine/Private/ChaosModularVehicle/ModularVehicleBuilder.cpp
+++ b/Engine/Plugins/Experimental/ChaosModularVehicle/Source/ChaosModularVehicleEngine/Private/ChaosModularVehicle/ModularVehicleBuilder.cpp
@@ -52,7 +52,7 @@ void FModularVehicleBuilder::FixupTreeLinks(TUniquePtr<Chaos::FSimModuleTree>& S
 				if (ParentModule && ParentModule->GetSimType() == eSimType::Wheel)
 				{
 					FSuspensionSimModule* Suspension = static_cast<FSuspensionSimModule*>(Module);
-					FWheelSimModule* Wheel = static_cast<FWheelSimModule*>(ParentModule);
+					FWheelBaseInterface* Wheel = static_cast<FWheelBaseInterface*>(ParentModule);
 					Wheel->SetSuspensionSimTreeIndex(Suspension->GetTreeIndex());
 					Suspension->SetWheelSimTreeIndex(Wheel->GetTreeIndex());
 				}
diff --git a/Engine/Plugins/Experimental/ChaosModularVehicle/Source/ChaosModularVehicleEngine/Private/ChaosModularVehicle/ModularVehicleSimulationCU.cpp b/Engine/Plugins/Experimental/ChaosModularVehicle/Source/ChaosModularVehicleEngine/Private/ChaosModularVehicle/ModularVehicleSimulationCU.cpp
--- a/Engine/Plugins/Experimental/ChaosModularVehicle/Source/ChaosModularVehicleEngine/Private/ChaosModularVehicle/ModularVehicleSimulationCU.cpp
+++ b/Engine/Plugins/Experimental/ChaosModularVehicle/Source/ChaosModularVehicleEngine/Private/ChaosModularVehicle/ModularVehicleSimulationCU.cpp
@@ -200,10 +200,10 @@ void FModularVehicleSimulationCU::PerformAdditionalSimWork(UWorld* InWorld, cons
 								float WheelRadius = 0;
 								if (Suspension->GetWheelSimTreeIndex() != Chaos::ISimulationModuleBase::INVALID_IDX)
 								{
-									Chaos::FWheelSimModule* Wheel = static_cast<Chaos::FWheelSimModule*>(ModuleArray[Suspension->GetWheelSimTreeIndex()].SimModule);
+									Chaos::FWheelBaseInterface* Wheel = static_cast<Chaos::FWheelBaseInterface*>(ModuleArray[Suspension->GetWheelSimTreeIndex()].SimModule);
 									if (Wheel)
 									{
-										WheelRadius = Wheel->Setup().Radius;
+										WheelRadius = Wheel->GetWheelRadius();
 									}
 								}
 
@@ -233,7 +233,7 @@ void FModularVehicleSimulationCU::PerformAdditionalSimWork(UWorld* InWorld, cons
 									{
 										const Chaos::FSimModuleTree::FSimModuleNode& WheelNode = ModuleArray[Suspension->GetWheelSimTreeIndex()];
 
-										Chaos::FWheelSimModule* Wheel = static_cast<Chaos::FWheelSimModule*>(WheelNode.SimModule);
+										Chaos::FWheelBaseInterface* Wheel = static_cast<Chaos::FWheelBaseInterface*>(WheelNode.SimModule);
 										if (Wheel && HitResult.PhysMaterial.IsValid())
 										{
 											if (GModularVehicleDebugParams.FrictionOverride > 0)
@@ -255,7 +255,7 @@ void FModularVehicleSimulationCU::PerformAdditionalSimWork(UWorld* InWorld, cons
 
 									if (Suspension->GetWheelSimTreeIndex() != Chaos::ISimulationModuleBase::INVALID_IDX)
 									{
-										Chaos::FWheelSimModule* Wheel = static_cast<Chaos::FWheelSimModule*>(ModuleArray[Suspension->GetWheelSimTreeIndex()].SimModule);
+										Chaos::FWheelBaseInterface* Wheel = static_cast<Chaos::FWheelBaseInterface*>(ModuleArray[Suspension->GetWheelSimTreeIndex()].SimModule);
 										if (Wheel)
 										{
 											if (GModularVehicleDebugParams.ShowWheelData)
diff --git a/Engine/Source/Runtime/Experimental/ChaosVehicles/ChaosVehiclesCore/Private/SimModule/AerofoilModule.cpp b/Engine/Source/Runtime/Experimental/ChaosVehicles/ChaosVehiclesCore/Private/SimModule/AerofoilModule.cpp
--- a/Engine/Source/Runtime/Experimental/ChaosVehicles/ChaosVehiclesCore/Private/SimModule/AerofoilModule.cpp
+++ b/Engine/Source/Runtime/Experimental/ChaosVehicles/ChaosVehiclesCore/Private/SimModule/AerofoilModule.cpp
@@ -5,12 +5,21 @@
 #include "VehicleUtility.h"
 
 #if VEHICLE_DEBUGGING_ENABLED
-UE_DISABLE_OPTIMIZATION
+UE_DISABLE_OPTIMIZATION_SHIP
 #endif
 
 namespace Chaos
 {
 
+	FAerofoilSimModule::FAerofoilSimModule(const FAerofoilSettings& Settings) : TSimModuleSettings<FAerofoilSettings>(Settings)
+		, CurrentAirDensity(RealWorldConsts::AirDensity())
+		, AngleOfAttack(0.f)
+		, ControlSurfaceAngle(0.f)
+		, AirflowNormal(FVector::ZeroVector)
+		, AerofoilId(0)
+	{
+	}
+
 	void FAerofoilSimModule::Simulate(float DeltaTime, const FAllInputs& Inputs, FSimModuleTree& VehicleModuleSystem)
 	{
 		float Altitude = Inputs.VehicleWorldTransform.GetLocation().Z;
@@ -153,5 +162,5 @@ namespace Chaos
 } // namespace Chaos
 
 #if VEHICLE_DEBUGGING_ENABLED
-UE_ENABLE_OPTIMIZATION
+UE_ENABLE_OPTIMIZATION_SHIP
 #endif
diff --git a/Engine/Source/Runtime/Experimental/ChaosVehicles/ChaosVehiclesCore/Private/SimModule/ClutchModule.cpp b/Engine/Source/Runtime/Experimental/ChaosVehicles/ChaosVehiclesCore/Private/SimModule/ClutchModule.cpp
--- a/Engine/Source/Runtime/Experimental/ChaosVehicles/ChaosVehiclesCore/Private/SimModule/ClutchModule.cpp
+++ b/Engine/Source/Runtime/Experimental/ChaosVehicles/ChaosVehiclesCore/Private/SimModule/ClutchModule.cpp
@@ -5,7 +5,7 @@
 #include "VehicleUtility.h"
 
 #if VEHICLE_DEBUGGING_ENABLED
-UE_DISABLE_OPTIMIZATION
+UE_DISABLE_OPTIMIZATION_SHIP
 #endif
 
 namespace Chaos
@@ -14,7 +14,11 @@ namespace Chaos
 	FClutchSimModule::FClutchSimModule(const FClutchSettings& Settings) : TSimModuleSettings<FClutchSettings>(Settings)
 		, ClutchValue(0.f)
 	{
-
+		if (!FModuleFactoryRegister::Get().ContainsFactory(GetSimType()))
+		{
+			static TSharedPtr<FSimFactoryModule<FClutchSimModuleDatas>> SharedFactory = MakeShared<FSimFactoryModule<FClutchSimModuleDatas>>(GetDebugName());
+			FModuleFactoryRegister::Get().RegisterFactory(GetSimType(), SharedFactory);
+		}
 	}
 
 	void FClutchSimModule::Simulate(float DeltaTime, const FAllInputs& Inputs, FSimModuleTree& VehicleModuleSystem)
@@ -63,5 +67,5 @@ namespace Chaos
 } // namespace Chaos
 
 #if VEHICLE_DEBUGGING_ENABLED
-UE_ENABLE_OPTIMIZATION
+UE_ENABLE_OPTIMIZATION_SHIP
 #endif
diff --git a/Engine/Source/Runtime/Experimental/ChaosVehicles/ChaosVehiclesCore/Private/SimModule/EngineModule.cpp b/Engine/Source/Runtime/Experimental/ChaosVehicles/ChaosVehiclesCore/Private/SimModule/EngineModule.cpp
--- a/Engine/Source/Runtime/Experimental/ChaosVehicles/ChaosVehiclesCore/Private/SimModule/EngineModule.cpp
+++ b/Engine/Source/Runtime/Experimental/ChaosVehicles/ChaosVehiclesCore/Private/SimModule/EngineModule.cpp
@@ -5,12 +5,25 @@
 #include "VehicleUtility.h"
 
 #if VEHICLE_DEBUGGING_ENABLED
-UE_DISABLE_OPTIMIZATION
+UE_DISABLE_OPTIMIZATION_SHIP
 #endif
 
 namespace Chaos
 {
 
+	FEngineSimModule::FEngineSimModule(const FEngineSettings& Settings) : TSimModuleSettings<FEngineSettings>(Settings)
+		, EngineIdleSpeed(RPMToOmega(Setup().IdleRPM))
+		, MaxEngineSpeed(RPMToOmega(Setup().MaxRPM))
+		, EngineStarted(true)
+	{
+		if (!FModuleFactoryRegister::Get().ContainsFactory(GetSimType()))
+		{
+			static TSharedPtr<FSimFactoryModule<FEngineSimModuleDatas>> SharedFactory = MakeShared<FSimFactoryModule<FEngineSimModuleDatas>>(GetDebugName());
+			FModuleFactoryRegister::Get().RegisterFactory(GetSimType(), SharedFactory);
+		}
+
+	}
+
 	void FEngineSimModule::Simulate(float DeltaTime, const FAllInputs& Inputs, FSimModuleTree& VehicleModuleSystem)
 	{
 		if (!EngineStarted)
@@ -108,5 +121,5 @@ namespace Chaos
 } // namespace Chaos
 
 #if VEHICLE_DEBUGGING_ENABLED
-UE_ENABLE_OPTIMIZATION
+UE_ENABLE_OPTIMIZATION_SHIP
 #endif
diff --git a/Engine/Source/Runtime/Experimental/ChaosVehicles/ChaosVehiclesCore/Private/SimModule/ModuleFactoryRegister.cpp b/Engine/Source/Runtime/Experimental/ChaosVehicles/ChaosVehiclesCore/Private/SimModule/ModuleFactoryRegister.cpp
new file mode 100644
--- /dev/null
+++ b/Engine/Source/Runtime/Experimental/ChaosVehicles/ChaosVehiclesCore/Private/SimModule/ModuleFactoryRegister.cpp
@@ -0,0 +1,54 @@
+// Copyright Epic Games, Inc. All Rights Reserved.
+
+#include "SimModule/ModuleFactoryRegister.h"
+#include "SimModule/SimulationModuleBase.h"
+
+namespace Chaos
+{
+	void FModuleFactoryRegister::RegisterFactory(int32 TypeID, TWeakPtr<IFactoryModule> InFactory)
+	{
+		if (!ContainsFactory(TypeID))
+		{
+			RegisteredFactories.Add(TypeID, InFactory);
+		}
+	}
+
+	void FModuleFactoryRegister::RemoveFactory(TWeakPtr<IFactoryModule> InFactory)
+	{
+		for (auto& It : RegisteredFactories)
+		{
+			if (It.Value == InFactory)
+			{
+				RegisteredFactories.Remove(It.Key);
+				return;
+			}
+		}
+	}
+
+	bool FModuleFactoryRegister::ContainsFactory(int32 TypeID)
+	{
+		return RegisteredFactories.Contains(TypeID);
+	}
+
+	TSharedPtr<Chaos::FModuleNetData> FModuleFactoryRegister::GenerateNetData(int32 TypeID, int32 SimArrayIndex)
+	{
+		using namespace Chaos;
+
+		if (RegisteredFactories.Contains(TypeID))
+		{
+			TSharedPtr<IFactoryModule> PinnedFactory = RegisteredFactories[TypeID].Pin();
+
+			if (PinnedFactory.IsValid())
+			{
+				return PinnedFactory->GenerateNetData(SimArrayIndex);
+			}
+		}
+		else
+		{
+			UE_LOG(LogTemp, Error, TEXT("No factory registered for type '%d'"), TypeID);
+		}
+
+		return nullptr;
+	}
+
+}
\ No newline at end of file
diff --git a/Engine/Source/Runtime/Experimental/ChaosVehicles/ChaosVehiclesCore/Private/SimModule/SuspensionModule.cpp b/Engine/Source/Runtime/Experimental/ChaosVehicles/ChaosVehiclesCore/Private/SimModule/SuspensionModule.cpp
--- a/Engine/Source/Runtime/Experimental/ChaosVehicles/ChaosVehiclesCore/Private/SimModule/SuspensionModule.cpp
+++ b/Engine/Source/Runtime/Experimental/ChaosVehicles/ChaosVehiclesCore/Private/SimModule/SuspensionModule.cpp
@@ -28,6 +28,13 @@ namespace Chaos
 		, WheelInContact(false)
 	{
 		AccessSetup().MaxLength = FMath::Abs(Settings.MaxRaise + Settings.MaxDrop);
+
+		if (!FModuleFactoryRegister::Get().ContainsFactory(GetSimType()))
+		{
+			static TSharedPtr<FSimFactoryModule<FSuspensionSimModuleDatas>> SharedFactory = MakeShared<FSimFactoryModule<FSuspensionSimModuleDatas>>(GetDebugName());
+			FModuleFactoryRegister::Get().RegisterFactory(GetSimType(), SharedFactory);
+		}
+
 	}
 
 	float FSuspensionSimModule::GetSpringLength() const
diff --git a/Engine/Source/Runtime/Experimental/ChaosVehicles/ChaosVehiclesCore/Private/SimModule/TransmissionModule.cpp b/Engine/Source/Runtime/Experimental/ChaosVehicles/ChaosVehiclesCore/Private/SimModule/TransmissionModule.cpp
--- a/Engine/Source/Runtime/Experimental/ChaosVehicles/ChaosVehiclesCore/Private/SimModule/TransmissionModule.cpp
+++ b/Engine/Source/Runtime/Experimental/ChaosVehicles/ChaosVehiclesCore/Private/SimModule/TransmissionModule.cpp
@@ -6,11 +6,24 @@
 #include "VehicleUtility.h"
 
 #if VEHICLE_DEBUGGING_ENABLED
-UE_DISABLE_OPTIMIZATION
+UE_DISABLE_OPTIMIZATION_SHIP
 #endif
 
 namespace Chaos
 {
+	FTransmissionSimModule::FTransmissionSimModule(const FTransmissionSettings& Settings)
+		: TSimModuleSettings<FTransmissionSettings>(Settings)
+		, CurrentGear(1)
+		, TargetGear(1)
+		, CurrentGearChangeTime(0.f)
+		, AllowedToChangeGear(true)
+	{
+		if (!FModuleFactoryRegister::Get().ContainsFactory(GetSimType()))
+		{
+			static TSharedPtr<FSimFactoryModule<FTransmissionSimModuleDatas>> SharedFactory = MakeShared<FSimFactoryModule<FTransmissionSimModuleDatas>>(GetDebugName());
+			FModuleFactoryRegister::Get().RegisterFactory(GetSimType(), SharedFactory);
+		}
+	}
 
 	void FTransmissionSimModule::Simulate(float DeltaTime, const FAllInputs& Inputs, FSimModuleTree& VehicleModuleSystem)
 	{
@@ -226,5 +239,5 @@ namespace Chaos
 } // namespace Chaos
 
 #if VEHICLE_DEBUGGING_ENABLED
-UE_ENABLE_OPTIMIZATION
+UE_ENABLE_OPTIMIZATION_SHIP
 #endif
diff --git a/Engine/Source/Runtime/Experimental/ChaosVehicles/ChaosVehiclesCore/Private/SimModule/WheelModule.cpp b/Engine/Source/Runtime/Experimental/ChaosVehicles/ChaosVehiclesCore/Private/SimModule/WheelModule.cpp
--- a/Engine/Source/Runtime/Experimental/ChaosVehicles/ChaosVehiclesCore/Private/SimModule/WheelModule.cpp
+++ b/Engine/Source/Runtime/Experimental/ChaosVehicles/ChaosVehiclesCore/Private/SimModule/WheelModule.cpp
@@ -14,14 +14,15 @@ namespace Chaos
 	FWheelSimModule::FWheelSimModule(const FWheelSettings& Settings)
 		: TSimModuleSettings<FWheelSettings>(Settings)
 		, BrakeTorque(0.0f)
-		, ForceIntoSurface(0.0f)
-		, SurfaceFriction(1.0f)
-		, SuspensionSimTreeIndex(INVALID_IDX)
 		, ForceFromFriction(FVector::ZeroVector)
 		, MassPerWheel(500.0f)
 		, SteerAngleDegrees(0.0f)
 	{
-
+		if (!FModuleFactoryRegister::Get().ContainsFactory(GetSimType()))
+		{
+			static TSharedPtr<FSimFactoryModule<FWheelSimModuleDatas>> SharedFactory = MakeShared<FSimFactoryModule<FWheelSimModuleDatas>>(GetDebugName());
+			FModuleFactoryRegister::Get().RegisterFactory(GetSimType(), SharedFactory);
+		}
 	}
 
 	void FWheelSimModule::Simulate(float DeltaTime, const FAllInputs& Inputs, FSimModuleTree& VehicleModuleSystem)
diff --git a/Engine/Source/Runtime/Experimental/ChaosVehicles/ChaosVehiclesCore/Public/SimModule/AerofoilModule.h b/Engine/Source/Runtime/Experimental/ChaosVehicles/ChaosVehiclesCore/Public/SimModule/AerofoilModule.h
--- a/Engine/Source/Runtime/Experimental/ChaosVehicles/ChaosVehiclesCore/Public/SimModule/AerofoilModule.h
+++ b/Engine/Source/Runtime/Experimental/ChaosVehicles/ChaosVehiclesCore/Public/SimModule/AerofoilModule.h
@@ -54,15 +54,7 @@ namespace Chaos
 	{
 	public:
 
-		FAerofoilSimModule(const FAerofoilSettings& Settings) : TSimModuleSettings<FAerofoilSettings>(Settings)
-			, CurrentAirDensity(RealWorldConsts::AirDensity())
-			, AngleOfAttack(0.f)
-			, ControlSurfaceAngle(0.f)
-			, AirflowNormal(FVector::ZeroVector)
-			, AerofoilId(0)
-		{
-
-		}
+		FAerofoilSimModule(const FAerofoilSettings& Settings);
 
 		virtual ~FAerofoilSimModule() {}
 
diff --git a/Engine/Source/Runtime/Experimental/ChaosVehicles/ChaosVehiclesCore/Public/SimModule/EngineModule.h b/Engine/Source/Runtime/Experimental/ChaosVehicles/ChaosVehiclesCore/Public/SimModule/EngineModule.h
--- a/Engine/Source/Runtime/Experimental/ChaosVehicles/ChaosVehiclesCore/Public/SimModule/EngineModule.h
+++ b/Engine/Source/Runtime/Experimental/ChaosVehicles/ChaosVehiclesCore/Public/SimModule/EngineModule.h
@@ -86,12 +86,7 @@ namespace Chaos
 	{
 	public:
 
-		FEngineSimModule(const FEngineSettings& Settings) : TSimModuleSettings<FEngineSettings>(Settings)
-			, EngineIdleSpeed(RPMToOmega(Setup().IdleRPM))
-			, MaxEngineSpeed(RPMToOmega(Setup().MaxRPM))
-			, EngineStarted(true)
-		{
-		}
+		FEngineSimModule(const FEngineSettings& Settings);
 
 		virtual ~FEngineSimModule() {}
 
diff --git a/Engine/Source/Runtime/Experimental/ChaosVehicles/ChaosVehiclesCore/Public/SimModule/ModuleFactoryRegister.h b/Engine/Source/Runtime/Experimental/ChaosVehicles/ChaosVehiclesCore/Public/SimModule/ModuleFactoryRegister.h
new file mode 100644
--- /dev/null
+++ b/Engine/Source/Runtime/Experimental/ChaosVehicles/ChaosVehiclesCore/Public/SimModule/ModuleFactoryRegister.h
@@ -0,0 +1,34 @@
+// Copyright Epic Games, Inc. All Rights Reserved.
+
+#pragma once
+
+#include "CoreMinimal.h"
+
+
+namespace Chaos
+{
+
+struct FModuleNetData;
+class IFactoryModule;
+
+class CHAOSVEHICLESCORE_API FModuleFactoryRegister
+{
+public:
+	static FModuleFactoryRegister& Get()
+	{
+		static FModuleFactoryRegister Instance;
+		return Instance;
+	}
+
+	void RegisterFactory(int32 TypeID, TWeakPtr<IFactoryModule> InFactory);
+	void RemoveFactory(TWeakPtr<IFactoryModule> InFactory);
+	bool ContainsFactory(int32 TypeID);
+	TSharedPtr<Chaos::FModuleNetData> GenerateNetData(int32 TypeID, int32 SimArrayIndex);
+
+protected:
+
+	FModuleFactoryRegister() = default;
+	TMap<int32, TWeakPtr<IFactoryModule>> RegisteredFactories;
+};
+
+} // namespace Chaos
diff --git a/Engine/Source/Runtime/Experimental/ChaosVehicles/ChaosVehiclesCore/Public/SimModule/SimulationModuleBase.h b/Engine/Source/Runtime/Experimental/ChaosVehicles/ChaosVehiclesCore/Public/SimModule/SimulationModuleBase.h
--- a/Engine/Source/Runtime/Experimental/ChaosVehicles/ChaosVehiclesCore/Public/SimModule/SimulationModuleBase.h
+++ b/Engine/Source/Runtime/Experimental/ChaosVehicles/ChaosVehiclesCore/Public/SimModule/SimulationModuleBase.h
@@ -6,7 +6,7 @@
 #include "Logging/LogMacros.h"
 #include "Chaos/ParticleHandleFwd.h"
 #include "Chaos/GeometryParticlesfwd.h"
-
+#include "SimModule/ModuleFactoryRegister.h"
 
 DECLARE_LOG_CATEGORY_EXTERN(LogSimulationModule, Warning, All);
 
@@ -123,13 +123,13 @@ namespace Chaos
 	};
 
 
-	enum eSimModuleState
+	enum CHAOSVEHICLESCORE_API eSimModuleState
 	{
 		Disabled,
 		Enabled
 	};
 
-	enum eSimModuleTypeFlags
+	enum CHAOSVEHICLESCORE_API eSimModuleTypeFlags
 	{
 		NonFunctional = (1 << 0),	// bitmask 1,2,4,8
 		Raycast = (1 << 1),	// requires raycast data
@@ -137,7 +137,7 @@ namespace Chaos
 		Velocity = (1 << 3),	// requires velocity data
 	};
 
-	enum eSimType
+	enum CHAOSVEHICLESCORE_API eSimType
 	{
 		Undefined = 0,
 		Chassis,		// linear/angular damping can be applied here
@@ -157,6 +157,12 @@ namespace Chaos
 		TorqueSim
 	};
 
+	enum CHAOSVEHICLESCORE_API EWheelAxis
+	{
+		X,	// X forward
+		Y	// Y forward
+	};
+
 	/**
 	 * Interface base class for all simulation module building blocks
 	 */
@@ -428,7 +434,40 @@ namespace Chaos
 #endif
 	};
 
+	class CHAOSVEHICLESCORE_API IFactoryModule
+	{
+	public:
+		virtual ~IFactoryModule() {}
+		virtual TSharedPtr<Chaos::FModuleNetData> GenerateNetData(int32 SimArrayIndex) const = 0;
 
-} // namespace Chaos
+	};
+
+	template <typename T>
+	class CHAOSVEHICLESCORE_API FSimFactoryModule : public IFactoryModule
+	{
+	public:
+		FSimFactoryModule(const FString& DebugNameIn)
+		{
+#if !(UE_BUILD_SHIPPING || UE_BUILD_TEST)
+			DebugString = DebugNameIn;
+#endif
+		}
 
 
+		TSharedPtr<Chaos::FModuleNetData> GenerateNetData(int32 SimArrayIndex) const
+		{
+			return MakeShared<T>(
+				SimArrayIndex
+	#if !(UE_BUILD_SHIPPING || UE_BUILD_TEST)
+				, DebugString
+	#endif			
+			);
+		}
+
+#if !(UE_BUILD_SHIPPING || UE_BUILD_TEST)
+		FString DebugString;
+#endif
+
+	};
+
+} // namespace Chaos
diff --git a/Engine/Source/Runtime/Experimental/ChaosVehicles/ChaosVehiclesCore/Public/SimModule/TorqueSimModule.h b/Engine/Source/Runtime/Experimental/ChaosVehicles/ChaosVehiclesCore/Public/SimModule/TorqueSimModule.h
--- a/Engine/Source/Runtime/Experimental/ChaosVehicles/ChaosVehiclesCore/Public/SimModule/TorqueSimModule.h
+++ b/Engine/Source/Runtime/Experimental/ChaosVehicles/ChaosVehiclesCore/Public/SimModule/TorqueSimModule.h
@@ -118,5 +118,30 @@ namespace Chaos
 		float AngularPosition;
 	};
 
+	class CHAOSVEHICLESCORE_API FWheelBaseInterface : public FTorqueSimModule
+	{
+	public:
+		FWheelBaseInterface()
+			: SuspensionSimTreeIndex(ISimulationModuleBase::INVALID_IDX)
+			, SurfaceFriction(1.0f)
+			, ForceIntoSurface(0.0f)
+		{}
+
+		void SetSuspensionSimTreeIndex(int IndexIn) { SuspensionSimTreeIndex = IndexIn; }
+		int GetSuspensionSimTreeIndex() const { return SuspensionSimTreeIndex; }
+		void SetSurfaceFriction(float FrictionIn) { SurfaceFriction = FrictionIn; }
+		void SetForceIntoSurface(float ForceIntoSurfaceIn) { ForceIntoSurface = ForceIntoSurfaceIn; }
+		float GetForceIntoSurface() const { return ForceIntoSurface; }
+		float GetSurfaceFriction() const { return SurfaceFriction; }
+
+		virtual float GetWheelRadius() const = 0;
+
+	protected:
+		int SuspensionSimTreeIndex;
+		float SurfaceFriction;
+		float ForceIntoSurface;
+
+	};
+
 
 } // namespace Chaos
diff --git a/Engine/Source/Runtime/Experimental/ChaosVehicles/ChaosVehiclesCore/Public/SimModule/TransmissionModule.h b/Engine/Source/Runtime/Experimental/ChaosVehicles/ChaosVehiclesCore/Public/SimModule/TransmissionModule.h
--- a/Engine/Source/Runtime/Experimental/ChaosVehicles/ChaosVehiclesCore/Public/SimModule/TransmissionModule.h
+++ b/Engine/Source/Runtime/Experimental/ChaosVehicles/ChaosVehiclesCore/Public/SimModule/TransmissionModule.h
@@ -105,14 +105,7 @@ namespace Chaos
 
 	public:
 
-		FTransmissionSimModule(const FTransmissionSettings& Settings)
-			: TSimModuleSettings<FTransmissionSettings>(Settings)
-			, CurrentGear(1)
-			, TargetGear(1)
-			, CurrentGearChangeTime(0.f)
-			, AllowedToChangeGear(true)
-		{
-		}
+		FTransmissionSimModule(const FTransmissionSettings& Settings);
 
 		virtual TSharedPtr<FModuleNetData> GenerateNetData(int SimArrayIndex) const
 		{
diff --git a/Engine/Source/Runtime/Experimental/ChaosVehicles/ChaosVehiclesCore/Public/SimModule/WheelModule.h b/Engine/Source/Runtime/Experimental/ChaosVehicles/ChaosVehiclesCore/Public/SimModule/WheelModule.h
--- a/Engine/Source/Runtime/Experimental/ChaosVehicles/ChaosVehiclesCore/Public/SimModule/WheelModule.h
+++ b/Engine/Source/Runtime/Experimental/ChaosVehicles/ChaosVehiclesCore/Public/SimModule/WheelModule.h
@@ -9,12 +9,7 @@ namespace Chaos
 	struct FAllInputs;
 	class FSimModuleTree;
 	class FClusterUnionPhysicsProxy;
-
-	enum CHAOSVEHICLESCORE_API EWheelAxis
-	{
-		X,	// X forward
-		Y	// Y forward
-	};
+	struct FModuleNetData;
 
 	struct CHAOSVEHICLESCORE_API FWheelSimModuleDatas : public FTorqueSimModuleDatas
 	{
@@ -120,7 +115,7 @@ namespace Chaos
 		bool ReverseDirection;
 	};
 
-	class CHAOSVEHICLESCORE_API FWheelSimModule : public FTorqueSimModule, public TSimModuleSettings<FWheelSettings>
+	class CHAOSVEHICLESCORE_API FWheelSimModule : public FWheelBaseInterface, public TSimModuleSettings<FWheelSettings>
 	{
 		friend FWheelOutputData;
 
@@ -155,16 +150,13 @@ namespace Chaos
 
 		virtual bool IsBehaviourType(eSimModuleTypeFlags InType) const override { return (InType & TorqueBased) || (InType & Velocity); }
 
-		void SetSuspensionSimTreeIndex(int IndexIn) { SuspensionSimTreeIndex = IndexIn; }
-		int GetSuspensionSimTreeIndex() const { return SuspensionSimTreeIndex; }
+		virtual float GetWheelRadius() const override { return Setup().Radius; }
 
 		float GetSteerAngleDegrees() const { return SteerAngleDegrees; }
 
-		void SetForceIntoSurface(float ForceIntoSurfaceIn) { ForceIntoSurface = ForceIntoSurfaceIn; }
-		float GetForceIntoSurface() const { return ForceIntoSurface; }
 		FVector GetForceFromFriction() const { return ForceFromFriction; }
-		void SetSurfaceFriction(float FrictionIn) { SurfaceFriction = FrictionIn; }
 		
+
 		/** set wheel rotational speed to match the specified linear forwards speed */
 		void SetLinearSpeed(float LinearMetersPerSecondIn)
 		{
@@ -186,9 +178,6 @@ namespace Chaos
 	private:
 
 		float BrakeTorque;				// [N.m]
-		float ForceIntoSurface;			// [N]
-		float SurfaceFriction;
-		int SuspensionSimTreeIndex;
 
 		FVector ForceFromFriction;
 		float MassPerWheel;
