Added option for sphere casts for suspension traces in modular vehicle

diff --git a/Engine/Plugins/Experimental/ChaosModularVehicle/Source/ChaosModularVehicleEngine/Private/ChaosModularVehicle/ModularVehicleBaseComponent.cpp b/Engine/Plugins/Experimental/ChaosModularVehicle/Source/ChaosModularVehicleEngine/Private/ChaosModularVehicle/ModularVehicleBaseComponent.cpp
--- a/Engine/Plugins/Experimental/ChaosModularVehicle/Source/ChaosModularVehicleEngine/Private/ChaosModularVehicle/ModularVehicleBaseComponent.cpp
+++ b/Engine/Plugins/Experimental/ChaosModularVehicle/Source/ChaosModularVehicleEngine/Private/ChaosModularVehicle/ModularVehicleBaseComponent.cpp
@@ -796,6 +796,7 @@ void UModularVehicleBaseComponent::Update(float DeltaTime)
 		TraceParams.bTraceComplex = bSuspensionTraceComplex;
 		AsyncInput->PhysicsInputs.TraceParams = TraceParams;
 		AsyncInput->PhysicsInputs.TraceCollisionResponse = SuspensionTraceCollisionResponses;
+		AsyncInput->PhysicsInputs.TraceType = TraceType;
 	}
 
 
diff --git a/Engine/Plugins/Experimental/ChaosModularVehicle/Source/ChaosModularVehicleEngine/Private/ChaosModularVehicle/ModularVehicleSimulationCU.cpp b/Engine/Plugins/Experimental/ChaosModularVehicle/Source/ChaosModularVehicleEngine/Private/ChaosModularVehicle/ModularVehicleSimulationCU.cpp
--- a/Engine/Plugins/Experimental/ChaosModularVehicle/Source/ChaosModularVehicleEngine/Private/ChaosModularVehicle/ModularVehicleSimulationCU.cpp
+++ b/Engine/Plugins/Experimental/ChaosModularVehicle/Source/ChaosModularVehicleEngine/Private/ChaosModularVehicle/ModularVehicleSimulationCU.cpp
@@ -221,7 +221,26 @@ void FModularVehicleSimulationCU::PerformAdditionalSimWork(UWorld* InWorld, cons
 								const FCollisionResponseParams& ResponseParams = InputData.PhysicsInputs.TraceCollisionResponse;
 								if (InWorld)
 								{
-									InWorld->LineTraceSingleByChannel(HitResult, TraceStart, TraceEnd, SpringCollisionChannel, TraceParams, ResponseParams);
+									switch (InputData.PhysicsInputs.TraceType)
+									{
+										case ETraceType::Spherecast:
+										{
+											InWorld->SweepSingleByChannel(HitResult
+												, TraceStart + TraceNormal * WheelRadius
+												, TraceEnd + TraceNormal * WheelRadius
+												, FQuat::Identity, SpringCollisionChannel
+												, FCollisionShape::MakeSphere(WheelRadius), TraceParams
+												, ResponseParams);
+										}
+										break;
+
+										case ETraceType::Raycast:
+										default:
+										{
+											InWorld->LineTraceSingleByChannel(HitResult, TraceStart, TraceEnd, SpringCollisionChannel, TraceParams, ResponseParams);
+										}
+										break;
+									}
 								}
 
 								float Offset = Suspension->Setup().MaxLength;
@@ -324,7 +343,19 @@ void FModularVehicleSimulationCU::PerformAdditionalSimWork(UWorld* InWorld, cons
 #endif
 								Suspension->SetSpringLength(Offset, WheelRadius);
 								FVector Up = ClusterWorldTM.GetUnitAxis(EAxis::Z);
-								Suspension->SetTargetPoint(HitResult.ImpactPoint + Up * WheelRadius, HitResult.ImpactNormal, HitResult.bBlockingHit);
+
+								FVector HitPoint;
+
+								if (InputData.PhysicsInputs.TraceType == ETraceType::Spherecast)
+								{
+									HitPoint = HitResult.Location;
+								}
+								else
+								{
+									HitPoint = HitResult.ImpactPoint + Up * WheelRadius;
+								}
+
+								Suspension->SetTargetPoint(HitPoint, HitResult.ImpactNormal, HitResult.bBlockingHit);
 							}
 
 						}
diff --git a/Engine/Plugins/Experimental/ChaosModularVehicle/Source/ChaosModularVehicleEngine/Public/ChaosModularVehicle/ChaosSimModuleManagerAsyncCallback.h b/Engine/Plugins/Experimental/ChaosModularVehicle/Source/ChaosModularVehicleEngine/Public/ChaosModularVehicle/ChaosSimModuleManagerAsyncCallback.h
--- a/Engine/Plugins/Experimental/ChaosModularVehicle/Source/ChaosModularVehicleEngine/Public/ChaosModularVehicle/ChaosSimModuleManagerAsyncCallback.h
+++ b/Engine/Plugins/Experimental/ChaosModularVehicle/Source/ChaosModularVehicleEngine/Public/ChaosModularVehicle/ChaosSimModuleManagerAsyncCallback.h
@@ -207,6 +207,15 @@ struct FGameStateInputs
 	//bool IsInAir
 };
 
+UENUM()
+enum class ETraceType : uint8
+{
+	/** Use ray to determine suspension length to ground */
+	Raycast		UMETA(DisplayName = "Raycast"),
+
+	/** Use sphere to determine suspension length to ground */
+	Spherecast	UMETA(DisplayName = "Spherecast"),
+};
 
 /**
  * Per Vehicle input State from Game Thread to Physics Thread
@@ -216,11 +225,13 @@ struct CHAOSMODULARVEHICLEENGINE_API FPhysicsModularVehicleInputs
 	FPhysicsModularVehicleInputs()
 		: TraceParams()
 		, TraceCollisionResponse()
+		, TraceType(ETraceType::Raycast)
 	{
 	}
 	mutable FNetworkModularVehicleInputs NetworkInputs;
 	mutable FCollisionQueryParams TraceParams;
 	mutable FCollisionResponseContainer TraceCollisionResponse;
+	mutable ETraceType TraceType;
 	mutable FGameStateInputs StateInputs;
 };
 
diff --git a/Engine/Plugins/Experimental/ChaosModularVehicle/Source/ChaosModularVehicleEngine/Public/ChaosModularVehicle/ModularVehicleBaseComponent.h b/Engine/Plugins/Experimental/ChaosModularVehicle/Source/ChaosModularVehicleEngine/Public/ChaosModularVehicle/ModularVehicleBaseComponent.h
--- a/Engine/Plugins/Experimental/ChaosModularVehicle/Source/ChaosModularVehicleEngine/Public/ChaosModularVehicle/ModularVehicleBaseComponent.h
+++ b/Engine/Plugins/Experimental/ChaosModularVehicle/Source/ChaosModularVehicleEngine/Public/ChaosModularVehicle/ModularVehicleBaseComponent.h
@@ -254,6 +254,10 @@ public:
 	UPROPERTY(EditAnywhere, Category = "Game|Components|ModularVehicle")
 	bool bSuspensionTraceComplex;
 
+	/** Wheel suspension trace type, defaults to ray trace */
+	UPROPERTY(EditAnywhere, Category = "Game|Components|ModularVehicle")
+	ETraceType TraceType;
+
 	UPROPERTY(EditAnywhere, Category = "Game|Components|ModularVehicle")
 	bool bKeepVehicleAwake;
 
