Fix invalid IsLocalControl state being passed to Network system.

diff --git a/Engine/Plugins/Experimental/ChaosModularVehicle/Source/ChaosModularVehicleEngine/Private/ChaosModularVehicle/ModularVehicleBaseComponent.cpp b/Engine/Plugins/Experimental/ChaosModularVehicle/Source/ChaosModularVehicleEngine/Private/ChaosModularVehicle/ModularVehicleBaseComponent.cpp
--- a/Engine/Plugins/Experimental/ChaosModularVehicle/Source/ChaosModularVehicleEngine/Private/ChaosModularVehicle/ModularVehicleBaseComponent.cpp
+++ b/Engine/Plugins/Experimental/ChaosModularVehicle/Source/ChaosModularVehicleEngine/Private/ChaosModularVehicle/ModularVehicleBaseComponent.cpp
@@ -56,19 +56,6 @@ UModularVehicleBaseComponent::UModularVehicleBaseComponent(const FObjectInitiali
 	EngineRPM = 0.0f;
 	EngineTorque = 0.0f;
 
-	//// #TODO: currently ordering of this must match EModularVehicleInputType
-	//InputInterpolationRates.Reset();
-	//InputInterpolationRates.Add(FModularVehicleInputRate(FString("Throttle")));
-	//InputInterpolationRates.Add(FModularVehicleInputRate(FString("Brake")));
-	//InputInterpolationRates.Add(FModularVehicleInputRate(FString("Clutch")));
-	//InputInterpolationRates.Add(FModularVehicleInputRate(FString("Steering")));
-	//InputInterpolationRates.Add(FModularVehicleInputRate(FString("Handbrake")));
-	//InputInterpolationRates.Add(FModularVehicleInputRate(FString("Pitch")));
-	//InputInterpolationRates.Add(FModularVehicleInputRate(FString("Roll")));
-	//InputInterpolationRates.Add(FModularVehicleInputRate(FString("Yaw")));
-	//InputInterpolationRates.Add(FModularVehicleInputRate(FString("Gear")));
-	//InputInterpolationRates.Add(FModularVehicleInputRate(FString("DebugIndex")));
-
 	if (bUsingNetworkPhysicsPrediction)
 	{
 		static const FName NetworkPhysicsComponentName(TEXT("PC_NetworkPhysicsComponent"));
@@ -111,7 +98,7 @@ APlayerController* UModularVehicleBaseComponent::GetPlayerController() const
 
 bool UModularVehicleBaseComponent::IsLocallyControlled() const
 {
-	if (bIsLocallyControlled && !GetWorld()->IsNetMode(NM_DedicatedServer))
+	if (bIsLocallyControlled)
 	{
 		return true;
 	}
@@ -978,7 +965,8 @@ void UModularVehicleBaseComponent::SetLocallyControlled(bool bLocallyControlledI
 	bIsLocallyControlled = false;
 	if (UWorld* World = GetWorld())
 	{
-		if (!World->IsNetMode(NM_DedicatedServer))
+		// guard against invalid case that can lead to bad networking state
+		if (GetOwner() && GetOwner()->GetLocalRole() != ENetRole::ROLE_SimulatedProxy)
 		{
 			bIsLocallyControlled = bLocallyControlledIn;
 		}
@@ -986,7 +974,7 @@ void UModularVehicleBaseComponent::SetLocallyControlled(bool bLocallyControlledI
 
 	if (bUsingNetworkPhysicsPrediction && NetworkPhysicsComponent)
 	{
-		NetworkPhysicsComponent->SetIsRelayingLocalInputs(bLocallyControlledIn);
+		NetworkPhysicsComponent->SetIsRelayingLocalInputs(bIsLocallyControlled);
 	}
 
 }
