Change module tree simulation order to be leaf first,
Fix code that was making a flat hierachy

diff --git a/Engine/Plugins/Experimental/ChaosModularVehicle/Source/ChaosModularVehicleEngine/Private/ChaosModularVehicle/ModularVehicleBaseComponent.cpp b/Engine/Plugins/Experimental/ChaosModularVehicle/Source/ChaosModularVehicleEngine/Private/ChaosModularVehicle/ModularVehicleBaseComponent.cpp
--- a/Engine/Plugins/Experimental/ChaosModularVehicle/Source/ChaosModularVehicleEngine/Private/ChaosModularVehicle/ModularVehicleBaseComponent.cpp
+++ b/Engine/Plugins/Experimental/ChaosModularVehicle/Source/ChaosModularVehicleEngine/Private/ChaosModularVehicle/ModularVehicleBaseComponent.cpp
@@ -182,7 +182,8 @@ int GenerateNewGuid()
 	return Val++;
 }
 
-void UModularVehicleBaseComponent::CreateAssociatedSimComponents(const UPrimitiveComponent* AttachedComponent, int ParentIndex, int TransformIndex, Chaos::FSimTreeUpdates& TreeUpdatesOut)
+
+void UModularVehicleBaseComponent::CreateAssociatedSimComponents(UPrimitiveComponent* AttachedComponent, int ParentIndex, int TransformIndex, Chaos::FSimTreeUpdates& TreeUpdatesOut)
 {
 	using namespace Chaos;
 	if (AttachedComponent == nullptr || ClusterUnionComponent == nullptr)
@@ -264,6 +265,11 @@ void UModularVehicleBaseComponent::CreateAssociatedSimComponents(const UPrimitiv
 			}
 		}
 
+		// store the tree index in the original sim component
+		if (UVehicleSimBaseComponent* SimComponent = Cast<UVehicleSimBaseComponent>(AttachedComponent))
+		{
+			SimComponent->TreeIndex = TreeIndex;
+		}
 		ParentIndex = TreeIndex;
 
 		if (bModularVehicle_SuspensionConstraint_Enabled)
@@ -730,7 +736,7 @@ void UModularVehicleBaseComponent::ParallelUpdate(float DeltaTime)
 
 					if (Chaos::FSimOutputData* ModuleOutput = PVehicleOutput->SimTreeOutputData[I])
 					{
-						if (ModuleOutput->AnimationSetupIndex >= 0)
+						if ((ModuleOutput->AnimationSetupIndex >= 0) && (ModuleOutput->AnimationSetupIndex < ModuleAnimationSetups.Num()))
 						{
 							ModuleAnimationSetups[ModuleOutput->AnimationSetupIndex].AnimFlags |= ModuleOutput->AnimFlags;
 
@@ -851,6 +857,25 @@ void UModularVehicleBaseComponent::ActionTreeUpdates(Chaos::FSimTreeUpdates* Nex
 		});
 }
 
+int UModularVehicleBaseComponent::FindParentsLastSimComponent(const UPrimitiveComponent* AttachedComponent)
+{
+	if (USceneComponent* AttachParent = AttachedComponent->GetAttachParent())
+	{
+		TArray<USceneComponent*> Children;
+		AttachParent->GetChildrenComponents(false, Children);
+
+		for (int I = Children.Num()-1; I>=0; I--)
+		{
+			if (UVehicleSimBaseComponent* ChildSimComponent = Cast<UVehicleSimBaseComponent>(Children[I]))
+			{
+				return ChildSimComponent->TreeIndex;
+			}
+		}
+	}
+
+	return INDEX_NONE;
+}
+
 void UModularVehicleBaseComponent::AddComponentToSimulation(UPrimitiveComponent* InComponent, const TArray<FClusterUnionBoneData>& BonesData, const TArray<FClusterUnionBoneData>& RemovedBoneIDs, bool bIsNew)
 {
 	check(ClusterUnionComponent);
@@ -874,7 +899,8 @@ void UModularVehicleBaseComponent::AddComponentToSimulation(UPrimitiveComponent*
 			ComponentAddOrder = FindComponentAddOrder(InComponent);
 		}
 
-		int ParentID = INDEX_NONE; // always at root
+		int ParentID = FindParentsLastSimComponent(InComponent);
+
 		Chaos::FSimTreeUpdates LatestTreeUpdates;
 		CreateAssociatedSimComponents(InComponent, ParentID, NextTransformIndex, LatestTreeUpdates);
 
diff --git a/Engine/Plugins/Experimental/ChaosModularVehicle/Source/ChaosModularVehicleEngine/Private/ChaosModularVehicle/ModularVehicleBuilder.cpp b/Engine/Plugins/Experimental/ChaosModularVehicle/Source/ChaosModularVehicleEngine/Private/ChaosModularVehicle/ModularVehicleBuilder.cpp
--- a/Engine/Plugins/Experimental/ChaosModularVehicle/Source/ChaosModularVehicleEngine/Private/ChaosModularVehicle/ModularVehicleBuilder.cpp
+++ b/Engine/Plugins/Experimental/ChaosModularVehicle/Source/ChaosModularVehicleEngine/Private/ChaosModularVehicle/ModularVehicleBuilder.cpp
@@ -57,6 +57,24 @@ void FModularVehicleBuilder::FixupTreeLinks(TUniquePtr<Chaos::FSimModuleTree>& S
 					Suspension->SetWheelSimTreeIndex(Wheel->GetTreeIndex());
 				}
 			}
+
+			for (int Child : SuspensionNode.Children)
+			{
+				if (Child != FSimModuleTree::FSimModuleNode::INVALID_IDX)
+				{
+					FSimModuleTree::FSimModuleNode& SuspensionChildNode = SimModuleTree->GetNode(Child);
+
+					ISimulationModuleBase* ChildModule = SimModuleTree->AccessSimModule(SuspensionChildNode.SimModule->GetTreeIndex());
+
+					if (ChildModule && ChildModule->GetSimType() == eSimType::Wheel)
+					{
+						FSuspensionSimModule* Suspension = static_cast<FSuspensionSimModule*>(Module);
+						FWheelBaseInterface* Wheel = static_cast<FWheelBaseInterface*>(ChildModule);
+						Wheel->SetSuspensionSimTreeIndex(Suspension->GetTreeIndex());
+						Suspension->SetWheelSimTreeIndex(Wheel->GetTreeIndex());
+					}
+				}
+			}
 		}
 
 		// temporarily - link suspension/wheels to engine/clutch/transmission, assuming only 1 transmission & all wheels are powered
diff --git a/Engine/Plugins/Experimental/ChaosModularVehicle/Source/ChaosModularVehicleEngine/Public/ChaosModularVehicle/ModularVehicleBaseComponent.h b/Engine/Plugins/Experimental/ChaosModularVehicle/Source/ChaosModularVehicleEngine/Public/ChaosModularVehicle/ModularVehicleBaseComponent.h
--- a/Engine/Plugins/Experimental/ChaosModularVehicle/Source/ChaosModularVehicleEngine/Public/ChaosModularVehicle/ModularVehicleBaseComponent.h
+++ b/Engine/Plugins/Experimental/ChaosModularVehicle/Source/ChaosModularVehicleEngine/Public/ChaosModularVehicle/ModularVehicleBaseComponent.h
@@ -210,7 +210,7 @@ public:
 	virtual void BeginPlay() override;
 	virtual void TickComponent(float DeltaTime, ELevelTick TickType, FActorComponentTickFunction* ThisTickFunction) override;
 
-	void CreateAssociatedSimComponents(const UPrimitiveComponent* AttachedComponent, int ParentIndex, int TransformIndex, Chaos::FSimTreeUpdates& TreeUpdatesOut);
+	void CreateAssociatedSimComponents(UPrimitiveComponent* AttachedComponent, int ParentIndex, int TransformIndex, Chaos::FSimTreeUpdates& TreeUpdatesOut);
 
 	void CreateConstraint(Chaos::ISimulationModuleBase* NewModule);
 	void DestroyConstraint(int ConstraintIndex);
@@ -384,6 +384,7 @@ protected:
 	void ActionTreeUpdates(Chaos::FSimTreeUpdates* NextTreeUpdates);
 
 	void SetCurrentAsyncDataInternal(FModularVehicleAsyncInput* CurInput, int32 InputIdx, FChaosSimModuleManagerAsyncOutput* CurOutput, FChaosSimModuleManagerAsyncOutput* NextOutput, float Alpha, int32 VehicleManagerTimestamp);
+	int FindParentsLastSimComponent(const UPrimitiveComponent* AttachedComponent);
 
 	IPhysicsProxyBase* GetPhysicsProxy() const;
 
diff --git a/Engine/Source/Runtime/Experimental/ChaosVehicles/ChaosVehiclesCore/Private/SimModule/SimModuleTree.cpp b/Engine/Source/Runtime/Experimental/ChaosVehicles/ChaosVehiclesCore/Private/SimModule/SimModuleTree.cpp
--- a/Engine/Source/Runtime/Experimental/ChaosVehicles/ChaosVehiclesCore/Private/SimModule/SimModuleTree.cpp
+++ b/Engine/Source/Runtime/Experimental/ChaosVehicles/ChaosVehiclesCore/Private/SimModule/SimModuleTree.cpp
@@ -95,9 +95,18 @@ void FSimModuleTree::AppendTreeUpdates(const FSimTreeUpdates& TreeUpdates)
 	for (const FPendingModuleAdds& TreeUpdate : TreeUpdates.GetNewModules())
 	{
 		int AddIndex = -1;
-		if (int* AddIndexPtr = SimTreeMapping.Find(TreeUpdate.ParentIndex))
+		if (LocalIndex == 0)
 		{
-			AddIndex = *AddIndexPtr;
+			// the first tree update contains the actual parent index in the real tree..
+			AddIndex = TreeUpdate.ParentIndex;
+		}
+		else
+		{
+			// all other tree updates have a parent index that is relative to the first
+			if (int* AddIndexPtr = SimTreeMapping.Find(TreeUpdate.ParentIndex))
+			{
+				AddIndex = *AddIndexPtr;
+			}
 		}
 
 		TreeIndex = AddNodeBelow(AddIndex, TreeUpdate.NewSimModule);
@@ -239,13 +248,16 @@ void FSimModuleTree::SimulateNode(float DeltaTime, FAllInputs& Inputs, int NodeI
 {
 	if (ISimulationModuleBase* Module = AccessSimModule(NodeIndex))
 	{
-		if (Module->IsEnabled())
+		if (SimTreeProcessingOrder == ESimTreeProcessingOrder::RootFirst)
 		{
-			Module->Simulate(PhysicsProxy, DeltaTime, Inputs, *this);
-
-			if (IsAnimationEnabled() && Module->IsAnimationEnabled())
+			if (Module->IsEnabled())
 			{
-				Module->Animate(PhysicsProxy);
+				Module->Simulate(PhysicsProxy, DeltaTime, Inputs, *this);
+
+				if (IsAnimationEnabled() && Module->IsAnimationEnabled())
+				{
+					Module->Animate(PhysicsProxy);
+				}
 			}
 		}
 
@@ -253,6 +265,20 @@ void FSimModuleTree::SimulateNode(float DeltaTime, FAllInputs& Inputs, int NodeI
 		{
 			SimulateNode(DeltaTime, Inputs, ChildIdx, PhysicsProxy);
 		}
+
+		if (SimTreeProcessingOrder == ESimTreeProcessingOrder::LeafFirst)
+		{
+			if (Module->IsEnabled())
+			{
+				Module->Simulate(PhysicsProxy, DeltaTime, Inputs, *this);
+
+				if (IsAnimationEnabled() && Module->IsAnimationEnabled())
+				{
+					Module->Animate(PhysicsProxy);
+				}
+			}
+		}
+
 	}
 }
 
diff --git a/Engine/Source/Runtime/Experimental/ChaosVehicles/ChaosVehiclesCore/Public/SimModule/SimModuleTree.h b/Engine/Source/Runtime/Experimental/ChaosVehicles/ChaosVehiclesCore/Public/SimModule/SimModuleTree.h
--- a/Engine/Source/Runtime/Experimental/ChaosVehicles/ChaosVehiclesCore/Public/SimModule/SimModuleTree.h
+++ b/Engine/Source/Runtime/Experimental/ChaosVehicles/ChaosVehiclesCore/Public/SimModule/SimModuleTree.h
@@ -8,6 +8,7 @@
 
 DECLARE_STATS_GROUP(TEXT("ModularVehicle.SimTree"), STATGROUP_ModularVehicleSimTree, STATGROUP_Advanced);
 
+
 class FGeometryCollectionPhysicsProxy;
 namespace Chaos
 {
@@ -16,6 +17,12 @@ namespace Chaos
 	struct FAllInputs;
 
 
+	enum ESimTreeProcessingOrder : int8
+	{
+		LeafFirst = 0,
+		RootFirst = 1
+	};
+
 	struct FPendingModuleAdds
 	{
 		FPendingModuleAdds(int ParentIndexIn, ISimulationModuleBase* NewSimModuleIn)
@@ -32,7 +39,6 @@ namespace Chaos
 	};
 
 	// Each update tree has it's own local tree hierarchy, this will be translated into the actual tree hierarchy.
-	// The root in here will always translate to the root in the main sim tree
 	class FSimTreeUpdates
 	{
 	public:
@@ -103,6 +109,7 @@ namespace Chaos
 		{
 			bAnimationEnabled = true;
 			bSimulationEnabled = true;
+			SimTreeProcessingOrder = ESimTreeProcessingOrder::LeafFirst;
 		}
 
 		~FSimModuleTree()
@@ -140,6 +147,8 @@ namespace Chaos
 
 		void Simulate(float DeltaTime, FAllInputs& Inputs, FClusterUnionPhysicsProxy* PhysicsProxy);
 
+		void SetSimTreeProcessingOrder(ESimTreeProcessingOrder OrderIn) { SimTreeProcessingOrder = OrderIn; }
+
 		FDeferredForcesModular& AccessDeferredForces() { return DeferredForces; }
 		const FDeferredForcesModular& GetDeferredForces() const { return DeferredForces; }
 		const TArray<FSimModuleNode>& GetSimulationModuleTree() { return SimulationModuleTree; }
@@ -215,6 +224,7 @@ namespace Chaos
 		bool bSimulationEnabled;
 
 		FVehicleState VehicleState;
+		ESimTreeProcessingOrder SimTreeProcessingOrder;
 	};
 
 
