Possbile fixes for modular vehicle networking crash

diff --git a/Engine/Plugins/Experimental/ChaosModularVehicle/Source/ChaosModularVehicleEngine/Private/ChaosModularVehicle/ChaosSimModuleManager.cpp b/Engine/Plugins/Experimental/ChaosModularVehicle/Source/ChaosModularVehicleEngine/Private/ChaosModularVehicle/ChaosSimModuleManager.cpp
--- a/Engine/Plugins/Experimental/ChaosModularVehicle/Source/ChaosModularVehicleEngine/Private/ChaosModularVehicle/ChaosSimModuleManager.cpp
+++ b/Engine/Plugins/Experimental/ChaosModularVehicle/Source/ChaosModularVehicleEngine/Private/ChaosModularVehicle/ChaosSimModuleManager.cpp
@@ -5,7 +5,7 @@
 #include "PBDRigidsSolver.h"
 #include "GameFramework/HUD.h" // for ShowDebugInfo
 #include "Physics/Experimental/PhysScene_Chaos.h"
-
+#include "SimModule/ModuleFactoryRegister.h"
 
 TMap<FPhysScene*, FChaosSimModuleManager*> FChaosSimModuleManager::SceneToModuleManagerMap;
 
@@ -70,6 +70,9 @@ void FChaosSimModuleManager::OnWorldCleanup(UWorld* InWorld, bool bSessionEnded,
 	{
 		Manager->UnregisterCallbacks();
 	}
+
+	Chaos::FModuleFactoryRegister::Get().Reset();
+
 }
 
 void FChaosSimModuleManager::OnShowDebugInfo(AHUD* HUD, UCanvas* Canvas, const FDebugDisplayInfo& DisplayInfo, float& YL, float& YPos)
diff --git a/Engine/Plugins/Experimental/ChaosModularVehicle/Source/ChaosModularVehicleEngine/Private/ChaosModularVehicle/ChaosSimModuleManagerAsyncCallback.cpp b/Engine/Plugins/Experimental/ChaosModularVehicle/Source/ChaosModularVehicleEngine/Private/ChaosModularVehicle/ChaosSimModuleManagerAsyncCallback.cpp
--- a/Engine/Plugins/Experimental/ChaosModularVehicle/Source/ChaosModularVehicleEngine/Private/ChaosModularVehicle/ChaosSimModuleManagerAsyncCallback.cpp
+++ b/Engine/Plugins/Experimental/ChaosModularVehicle/Source/ChaosModularVehicleEngine/Private/ChaosModularVehicle/ChaosSimModuleManagerAsyncCallback.cpp
@@ -271,22 +271,18 @@ bool FNetworkModularVehicleStates::NetSerialize(FArchive& Ar, class UPackageMap*
 				Ar << ModuleType;
 				Ar << SimArrayIndex;
 
-				if (!ModuleData.IsEmpty())
-				{
-					ensure(I <= ModuleData.Num());
-				}
-
 				if (ModuleData.Num() != NumNetModules)
 				{
 					ModuleData.Reserve(NumNetModules);
 
 					if (TSharedPtr<Chaos::FModuleNetData> Data = Chaos::FModuleFactoryRegister::Get().GenerateNetData(ModuleType, SimArrayIndex))
 					{
+						Data->Serialize(Ar);
 						ModuleData.Emplace(Data);
 					}
 					else
 					{
-						checkf(false, TEXT("Module net data factory has not been registered, use FModuleFactoryRegister::Get().RegisterFactory()"));
+						ensureMsgf(false, TEXT("Module net data factory has not been registered, use FModuleFactoryRegister::Get().RegisterFactory()"));
 					}
 
 				}
@@ -297,9 +293,9 @@ bool FNetworkModularVehicleStates::NetSerialize(FArchive& Ar, class UPackageMap*
 			int32 ModuleType = (int32)ModuleData[I]->GetType();
 			Ar << ModuleType;
 			Ar << ModuleData[I]->SimArrayIndex;
+			ModuleData[I]->Serialize(Ar);
 		}
 
-		ModuleData[I]->Serialize(Ar);
 	}
 
 	return true;
diff --git a/Engine/Source/Runtime/Experimental/ChaosVehicles/ChaosVehiclesCore/Private/SimModule/ModuleFactoryRegister.cpp b/Engine/Source/Runtime/Experimental/ChaosVehicles/ChaosVehiclesCore/Private/SimModule/ModuleFactoryRegister.cpp
--- a/Engine/Source/Runtime/Experimental/ChaosVehicles/ChaosVehiclesCore/Private/SimModule/ModuleFactoryRegister.cpp
+++ b/Engine/Source/Runtime/Experimental/ChaosVehicles/ChaosVehiclesCore/Private/SimModule/ModuleFactoryRegister.cpp
@@ -25,6 +25,11 @@ namespace Chaos
 		}
 	}
 
+	void FModuleFactoryRegister::Reset()
+	{
+		RegisteredFactories.Reset();
+	}
+
 	bool FModuleFactoryRegister::ContainsFactory(int32 TypeID)
 	{
 		return RegisteredFactories.Contains(TypeID);
diff --git a/Engine/Source/Runtime/Experimental/ChaosVehicles/ChaosVehiclesCore/Public/SimModule/ModuleFactoryRegister.h b/Engine/Source/Runtime/Experimental/ChaosVehicles/ChaosVehiclesCore/Public/SimModule/ModuleFactoryRegister.h
--- a/Engine/Source/Runtime/Experimental/ChaosVehicles/ChaosVehiclesCore/Public/SimModule/ModuleFactoryRegister.h
+++ b/Engine/Source/Runtime/Experimental/ChaosVehicles/ChaosVehiclesCore/Public/SimModule/ModuleFactoryRegister.h
@@ -22,6 +22,7 @@ public:
 
 	void RegisterFactory(int32 TypeID, TWeakPtr<IFactoryModule> InFactory);
 	void RemoveFactory(TWeakPtr<IFactoryModule> InFactory);
+	void Reset();
 	bool ContainsFactory(int32 TypeID);
 	TSharedPtr<Chaos::FModuleNetData> GenerateNetData(int32 TypeID, int32 SimArrayIndex);
 
