Modular vehicle control input refactor

diff --git a/Engine/Plugins/Experimental/ChaosModularVehicle/Source/ChaosModularVehicleEngine/Private/ChaosModularVehicle/ChaosSimModuleManagerAsyncCallback.cpp b/Engine/Plugins/Experimental/ChaosModularVehicle/Source/ChaosModularVehicleEngine/Private/ChaosModularVehicle/ChaosSimModuleManagerAsyncCallback.cpp
--- a/Engine/Plugins/Experimental/ChaosModularVehicle/Source/ChaosModularVehicleEngine/Private/ChaosModularVehicle/ChaosSimModuleManagerAsyncCallback.cpp
+++ b/Engine/Plugins/Experimental/ChaosModularVehicle/Source/ChaosModularVehicleEngine/Private/ChaosModularVehicle/ChaosSimModuleManagerAsyncCallback.cpp
@@ -12,7 +12,6 @@
 
 FSimModuleDebugParams GSimModuleDebugParams;
 
-
 DECLARE_CYCLE_STAT(TEXT("AsyncCallback:OnPreSimulate_Internal"), STAT_AsyncCallback_OnPreSimulate, STATGROUP_ChaosSimModuleManager);
 
 FName FChaosSimModuleManagerAsyncCallback::GetFNameForStatId() const
@@ -188,25 +187,17 @@ void FModularVehicleAsyncInput::ProcessInputs()
 	{
 		PhysicsInputs.NetworkInputs.VehicleInputs = VehicleSim->VehicleInputs;
 	}
-
 }
 
 bool FNetworkModularVehicleInputs::NetSerialize(FArchive& Ar, class UPackageMap* Map, bool& bOutSuccess)
 {
 	FNetworkPhysicsData::SerializeFrames(Ar);
 
-	Ar << VehicleInputs.Steering;
-	Ar << VehicleInputs.Throttle;
-	Ar << VehicleInputs.Brake;
-	Ar << VehicleInputs.Handbrake;
-	Ar << VehicleInputs.Pitch;
-	Ar << VehicleInputs.Roll;
-	Ar << VehicleInputs.Yaw;
-	Ar << VehicleInputs.Boost;
-	Ar << VehicleInputs.Drift;
 	Ar << VehicleInputs.Reverse;
 	Ar << VehicleInputs.KeepAwake;
 
+	VehicleInputs.Container.Serialize(Ar, Map, bOutSuccess);
+
 	bOutSuccess = true;
 	return bOutSuccess;
 }
@@ -240,17 +231,15 @@ void FNetworkModularVehicleInputs::InterpolateData(const FNetworkPhysicsData& Mi
 
 	const float LerpFactor = (LocalFrame - MinInput.LocalFrame) / (MaxInput.LocalFrame - MinInput.LocalFrame);
 
-	VehicleInputs.Steering = FMath::Lerp(MinInput.VehicleInputs.Steering, MaxInput.VehicleInputs.Steering, LerpFactor);
-	VehicleInputs.Throttle = FMath::Lerp(MinInput.VehicleInputs.Throttle, MaxInput.VehicleInputs.Throttle, LerpFactor);
-	VehicleInputs.Brake = FMath::Lerp(MinInput.VehicleInputs.Brake, MaxInput.VehicleInputs.Brake, LerpFactor);
-	VehicleInputs.Handbrake = FMath::Lerp(MinInput.VehicleInputs.Handbrake, MaxInput.VehicleInputs.Handbrake, LerpFactor);
-	VehicleInputs.Pitch = FMath::Lerp(MinInput.VehicleInputs.Pitch, MaxInput.VehicleInputs.Pitch, LerpFactor);
-	VehicleInputs.Roll = FMath::Lerp(MinInput.VehicleInputs.Roll, MaxInput.VehicleInputs.Roll, LerpFactor);
-	VehicleInputs.Yaw = FMath::Lerp(MinInput.VehicleInputs.Yaw, MaxInput.VehicleInputs.Yaw, LerpFactor);
-	VehicleInputs.Boost = FMath::Lerp(MinInput.VehicleInputs.Boost, MaxInput.VehicleInputs.Boost, LerpFactor);
-	VehicleInputs.Drift = FMath::Lerp(MinInput.VehicleInputs.Drift, MaxInput.VehicleInputs.Drift, LerpFactor);
 	VehicleInputs.Reverse = MinInput.VehicleInputs.Reverse;
 	VehicleInputs.KeepAwake = MinInput.VehicleInputs.KeepAwake;
+	VehicleInputs.Container.Lerp(MinInput.VehicleInputs.Container, MaxInput.VehicleInputs.Container, LerpFactor);
+}
+
+void FNetworkModularVehicleInputs::MergeData(const FNetworkPhysicsData& FromData)
+{
+	const FNetworkModularVehicleInputs& FromInput = static_cast<const FNetworkModularVehicleInputs&>(FromData);
+	VehicleInputs.Container.Merge(FromInput.VehicleInputs.Container);
 }
 
 bool FNetworkModularVehicleStates::NetSerialize(FArchive& Ar, class UPackageMap* Map, bool& bOutSuccess)
@@ -280,10 +269,6 @@ bool FNetworkModularVehicleStates::NetSerialize(FArchive& Ar, class UPackageMap*
 						Data->Serialize(Ar);
 						ModuleData.Emplace(Data);
 					}
-					else
-					{
-						ensureMsgf(false, TEXT("Module net data factory has not been registered, use FModuleFactoryRegister::Get().RegisterFactory()"));
-					}
 
 				}
 			}
diff --git a/Engine/Plugins/Experimental/ChaosModularVehicle/Source/ChaosModularVehicleEngine/Private/ChaosModularVehicle/ModularVehicleBaseComponent.cpp b/Engine/Plugins/Experimental/ChaosModularVehicle/Source/ChaosModularVehicleEngine/Private/ChaosModularVehicle/ModularVehicleBaseComponent.cpp
--- a/Engine/Plugins/Experimental/ChaosModularVehicle/Source/ChaosModularVehicleEngine/Private/ChaosModularVehicle/ModularVehicleBaseComponent.cpp
+++ b/Engine/Plugins/Experimental/ChaosModularVehicle/Source/ChaosModularVehicleEngine/Private/ChaosModularVehicle/ModularVehicleBaseComponent.cpp
@@ -23,6 +23,7 @@
 
 DEFINE_LOG_CATEGORY(LogModularBase);
 
+
 bool bModularVehicle_SuspensionConstraint_Enabled = true;
 FAutoConsoleVariableRef CVarModularVehicleSuspensionConstraintEnabled(TEXT("p.ModularVehicle.SuspensionConstraint.Enabled"), bModularVehicle_SuspensionConstraint_Enabled, TEXT("Enable/Disable suspension constraint falling back to simple forces when constraint disabled (requires restart)."));
 
@@ -110,6 +111,38 @@ bool UModularVehicleBaseComponent::IsLocallyControlled() const
 	return false;
 }
 
+void UModularVehicleBaseComponent::GenerateInputModifiers(const TArray<FModuleInputSetup>& CombinedInputConfiguration)
+{
+	for (const FModuleInputSetup& InputSetup : CombinedInputConfiguration)
+	{
+		if (InputSetup.InputModifierClass != nullptr)
+		{
+			UDefaultModularVehicleInputModifier* NewPtr = NewObject<UDefaultModularVehicleInputModifier>(this, InputSetup.InputModifierClass);
+			InputModifiers.Add(NewPtr);
+		}
+		else
+		{
+			InputModifiers.Add(nullptr);
+		}
+	}
+}
+
+void UModularVehicleBaseComponent::ApplyInputModifiers(float DeltaTime, const FModuleInputContainer& RawValue)
+{
+	check(InputModifiers.Num() == InputsContainer.GetNumInputs());
+	for (int I = 0; I < InputsContainer.GetNumInputs(); I++)
+	{
+		if (InputModifiers[I] != nullptr)
+		{
+			InputsContainer.SetValueAtIndex(I, InputModifiers[I]->InterpInputValue(DeltaTime, InputsContainer.GetValueAtIndex(I), RawValue.GetValueAtIndex(I)));
+		}
+		else
+		{
+			InputsContainer.SetValueAtIndex(I, RawValue.GetValueAtIndex(I));
+		}
+	}
+}
+
 
 void UModularVehicleBaseComponent::OnCreatePhysicsState()
 {
@@ -140,6 +173,7 @@ void UModularVehicleBaseComponent::OnCreatePhysicsState()
 		{
 			// register interface to handle network prediction callbacks
 			// #Note: in our case we don't yet know what the replication data will be since the modules are built after this point at runtime
+			FScopedModuleInputInitializer SetSetup(InputConfig);
 			NetworkPhysicsComponent->CreateDataHistory<FPhysicsModularVehicleTraits>(this);
 
 			if (bIsLocallyControlled)
@@ -333,6 +367,21 @@ void UModularVehicleBaseComponent::BeginPlay()
 			AddGeometryCollectionsFromOwnedActor();
 		}
 	}
+
+	// control input setup - unfortunately can't do this in OnCreatePhysics since RootComponent->GetChildrenComponents will not work 
+	// at that time and AssimilateComponentInputs will not find any controls in the component hierarchy
+	TArray<FModuleInputSetup> CombinedInputConfiguration;
+	AssimilateComponentInputs(CombinedInputConfiguration);
+	RawInputsContainer.Initialize(CombinedInputConfiguration, InputNameMap);
+	InputsContainer = RawInputsContainer;
+	if (!bUsingNetworkPhysicsPrediction)
+	{
+		ReplicatedState.Container = InputsContainer;
+	}
+	GenerateInputModifiers(CombinedInputConfiguration);
+
+	VehicleSimulationPT->SetInputMappings(InputNameMap);
+
 }
 
 void UModularVehicleBaseComponent::TickComponent(float DeltaTime, ELevelTick TickType, FActorComponentTickFunction* ThisTickFunction)
@@ -513,22 +562,12 @@ void UModularVehicleBaseComponent::UpdateState(float DeltaTime)
 	// Should we remove input instead of relying on replicated state in that case?
 	if (bProcessLocally && PVehicleOutput)
 	{
-		// Apply Inputs locally
-		SteeringInput = SteeringInputRate.InterpInputValue(DeltaTime, SteeringInput, RawSteeringInput);
-		ThrottleInput = ThrottleInputRate.InterpInputValue(DeltaTime, ThrottleInput, RawThrottleInput);
-		BrakeInput = BrakeInputRate.InterpInputValue(DeltaTime, BrakeInput, RawBrakeInput);
-		HandbrakeInput = HandbrakeInputRate.InterpInputValue(DeltaTime, HandbrakeInput, RawHandbrakeInput);
-		PitchInput = PitchInputRate.InterpInputValue(DeltaTime, PitchInput, RawPitchInput);
-		RollInput = RollInputRate.InterpInputValue(DeltaTime, RollInput, RawRollInput);
-		YawInput = YawInputRate.InterpInputValue(DeltaTime, YawInput, RawYawInput);
-		BoostInput = BoostInputRate.InterpInputValue(DeltaTime, BoostInput, RawBoostInput);
-		DriftInput = DriftInputRate.InterpInputValue(DeltaTime, DriftInput, RawDriftInput);
-		ReverseInput = RawReverseInput;
+		ApplyInputModifiers(DeltaTime, RawInputsContainer);
 
 		if (!bUsingNetworkPhysicsPrediction)
 		{
 			// and send to server - (ServerUpdateState_Implementation below)
-			ServerUpdateState(SteeringInput, ThrottleInput, BrakeInput, HandbrakeInput, -1, RollInput, PitchInput, YawInput, BoostInput, DriftInput, ReverseInput);
+			ServerUpdateState(InputsContainer, ReverseInput, bKeepVehicleAwake);
 		}
 
 		if (PawnOwner && PawnOwner->IsNetMode(NM_Client))
@@ -539,49 +578,23 @@ void UModularVehicleBaseComponent::UpdateState(float DeltaTime)
 	else if (!bUsingNetworkPhysicsPrediction)
 	{
 		// use replicated values for remote pawns
-		ThrottleInput = ReplicatedState.Throttle;
-		SteeringInput = ReplicatedState.Steering;
-		BrakeInput = ReplicatedState.Brake;
-		HandbrakeInput = ReplicatedState.Handbrake;
-		PitchInput = ReplicatedState.Pitch;
-		RollInput = ReplicatedState.Roll;
-		YawInput = ReplicatedState.Yaw;
-		BoostInput = ReplicatedState.Boost;
-		DriftInput = ReplicatedState.Drift;
+		InputsContainer = ReplicatedState.Container;
 		ReverseInput = ReplicatedState.Reverse;
+		bKeepVehicleAwake = ReplicatedState.KeepAwake;
 	}
 }
 
-bool UModularVehicleBaseComponent::ServerUpdateState_Validate(float InSteeringInput, float InThrottleInput, float InBrakeInput, float InHandbrakeInput, int32 InCurrentGear, float InRollInput, float InPitchInput, float InYawInput, float InBoostInput, float InDriftInput, bool InReverseInput)
+bool UModularVehicleBaseComponent::ServerUpdateState_Validate(const FModuleInputContainer& InputsIn, bool Reverse, bool KeepAwake)
 {
 	return true;
 }
 
-void UModularVehicleBaseComponent::ServerUpdateState_Implementation(float InSteeringInput, float InThrottleInput, float InBrakeInput
-	, float InHandbrakeInput, int32 InCurrentGear, float InRollInput, float InPitchInput, float InYawInput, float InBoostInput, float InDriftInput, bool InReverseInput)
+void UModularVehicleBaseComponent::ServerUpdateState_Implementation(const FModuleInputContainer& InputsIn, bool Reverse, bool KeepAwake)
 {
-	SteeringInput = InSteeringInput;
-	ThrottleInput = InThrottleInput;
-	BrakeInput = InBrakeInput;
-	HandbrakeInput = InHandbrakeInput;
-	RollInput = InRollInput;
-	PitchInput = InPitchInput;
-	YawInput = InYawInput;
-	BoostInput = InBoostInput;
-	DriftInput = InDriftInput;
-	ReverseInput = InReverseInput;
-
 	// update state of inputs
-	ReplicatedState.Steering = InSteeringInput;
-	ReplicatedState.Throttle = InThrottleInput;
-	ReplicatedState.Brake = InBrakeInput;
-	ReplicatedState.Handbrake = InHandbrakeInput;
-	ReplicatedState.Roll = InRollInput;
-	ReplicatedState.Pitch = InPitchInput;
-	ReplicatedState.Yaw = InYawInput;
-	ReplicatedState.Boost = InBoostInput;
-	ReplicatedState.Drift = InDriftInput;
-	ReplicatedState.Reverse = InReverseInput;
+	ReplicatedState.Reverse = Reverse;
+	ReplicatedState.KeepAwake = KeepAwake;
+	ReplicatedState.Container = InputsIn;
 }
 
 
@@ -771,18 +784,12 @@ void UModularVehicleBaseComponent::Update(float DeltaTime)
 
 		FModularVehicleAsyncInput* AsyncInput = static_cast<FModularVehicleAsyncInput*>(CurAsyncInput);
 
-		AsyncInput->PhysicsInputs.NetworkInputs.VehicleInputs.Throttle = ThrottleInput;
-		AsyncInput->PhysicsInputs.NetworkInputs.VehicleInputs.Boost = BoostInput;
-		AsyncInput->PhysicsInputs.NetworkInputs.VehicleInputs.Drift = DriftInput;
-		AsyncInput->PhysicsInputs.NetworkInputs.VehicleInputs.Brake = BrakeInput;
-		AsyncInput->PhysicsInputs.NetworkInputs.VehicleInputs.Steering = SteeringInput;
-		AsyncInput->PhysicsInputs.NetworkInputs.VehicleInputs.Handbrake = HandbrakeInput;
-		AsyncInput->PhysicsInputs.NetworkInputs.VehicleInputs.Roll = RollInput;
-		AsyncInput->PhysicsInputs.NetworkInputs.VehicleInputs.Pitch = PitchInput;
-		AsyncInput->PhysicsInputs.NetworkInputs.VehicleInputs.Yaw = YawInput;
 		AsyncInput->PhysicsInputs.NetworkInputs.VehicleInputs.Reverse = ReverseInput;
 		AsyncInput->PhysicsInputs.NetworkInputs.VehicleInputs.KeepAwake = bKeepVehicleAwake;
 
+		// All control inputs
+		AsyncInput->PhysicsInputs.NetworkInputs.VehicleInputs.Container = InputsContainer;
+
 		FCollisionQueryParams TraceParams(NAME_None, FCollisionQueryParams::GetUnknownStatId(), false, nullptr);
 		TraceParams.bReturnPhysicalMaterial = true;	// we need this to get the surface friction coefficient
 		TraceParams.AddIgnoredActors(ActorsToIgnore);
@@ -1005,6 +1012,37 @@ void UModularVehicleBaseComponent::SetLocallyControlled(bool bLocallyControlledI
 
 }
 
+void UModularVehicleBaseComponent::AssimilateComponentInputs(TArray<FModuleInputSetup>& OutCombinedInputs)
+{
+	// copy the input setup from this class
+	OutCombinedInputs = InputConfig;
+
+	// append the input setup from all module sim components attached to same actor
+	if (APawn* Pawn = Cast<APawn>(GetOwner()))
+	{
+		if (USceneComponent* RootComponent = Pawn->GetRootComponent())
+		{
+			TArray<USceneComponent*> ChildComponents;
+			RootComponent->GetChildrenComponents(true, ChildComponents);
+
+			for (USceneComponent* Component : ChildComponents)
+			{
+				if (UVehicleSimBaseComponent* GCComponent = Cast<UVehicleSimBaseComponent>(Component))
+				{
+					// don't add duplicates, i.e. 4 wheels could be looking for a single steering input
+					for (FModuleInputSetup& Config : GCComponent->InputConfig)
+					{
+						if (OutCombinedInputs.Find(Config) == INDEX_NONE)
+						{
+							OutCombinedInputs.Append(GCComponent->InputConfig);
+						}
+					}
+				}
+			}
+		}
+	}
+}
+
 void UModularVehicleBaseComponent::CreateVehicleSim()
 {
 	UWorld* World = GetWorld();
@@ -1119,59 +1157,52 @@ void UModularVehicleBaseComponent::AddGeometryCollectionsFromOwnedActor()
 	}
 }
 
-void UModularVehicleBaseComponent::SetThrottleInput(float Throttle)
-{
-	RawThrottleInput = FMath::Clamp(Throttle, -1.0f, 1.0f);
-}
-
-void UModularVehicleBaseComponent::SetBoostInput(float Boost)
+void UModularVehicleBaseComponent::SetInputBool(const FName Name, const bool Value)
 {
-	RawBoostInput = FMath::Clamp(Boost, -1.0f, 1.0f);
+	FInputInterface Inputs(InputNameMap, RawInputsContainer);
+	Inputs.SetValue(Name, Value);
 }
 
-void UModularVehicleBaseComponent::SetDriftInput(float Drift)
+void UModularVehicleBaseComponent::SetInputAxis1D(const FName Name, const double Value)
 {
-	RawDriftInput = FMath::Clamp(Drift, -1.0f, 1.0f);
+	FInputInterface Inputs(InputNameMap, RawInputsContainer);
+	Inputs.SetValue(Name, Value);
 }
 
-void UModularVehicleBaseComponent::IncreaseThrottleInput(float ThrottleDelta)
+void UModularVehicleBaseComponent::SetInputAxis2D(const FName Name, const FVector2D Value)
 {
-	RawThrottleInput = FMath::Clamp(RawThrottleInput + ThrottleDelta, 0.f, 1.0f);
+	FInputInterface Inputs(InputNameMap, RawInputsContainer);
+	Inputs.SetValue(Name, Value);
 }
 
-void UModularVehicleBaseComponent::DecreaseThrottleInput(float ThrottleDelta)
+void UModularVehicleBaseComponent::SetInputAxis3D(const FName Name, const FVector Value)
 {
-	RawThrottleInput = FMath::Clamp(RawThrottleInput - ThrottleDelta, 0.f, 1.0f);
+	FInputInterface Inputs(InputNameMap, RawInputsContainer);
+	Inputs.SetValue(Name, Value);
 }
 
-void UModularVehicleBaseComponent::SetBrakeInput(float Brake)
+void UModularVehicleBaseComponent::SetInput(const FName& Name, const bool Value)
 {
-	RawBrakeInput = FMath::Clamp(Brake, -1.0f, 1.0f);
+	FInputInterface Inputs(InputNameMap, RawInputsContainer);
+	Inputs.SetValue(Name, Value);
 }
 
-void UModularVehicleBaseComponent::SetSteeringInput(float Steering)
+void UModularVehicleBaseComponent::SetInput(const FName& Name, const double Value)
 {
-	RawSteeringInput = FMath::Clamp(Steering, -1.0f, 1.0f);
+	FInputInterface Inputs(InputNameMap, RawInputsContainer);
+	Inputs.SetValue(Name, Value);
 }
 
-void UModularVehicleBaseComponent::SetPitchInput(float Pitch)
+void UModularVehicleBaseComponent::SetInput(const FName& Name, const FVector2D& Value)
 {
-	RawPitchInput = FMath::Clamp(Pitch, -1.0f, 1.0f);
+	FInputInterface Inputs(InputNameMap, RawInputsContainer);
+	Inputs.SetValue(Name, Value);
 }
 
-void UModularVehicleBaseComponent::SetRollInput(float Roll)
+void UModularVehicleBaseComponent::SetInput(const FName& Name, const FVector& Value)
 {
-	RawRollInput = FMath::Clamp(Roll, -1.0f, 1.0f);
-}
-
-void UModularVehicleBaseComponent::SetYawInput(float Yaw)
-{
-	RawYawInput = FMath::Clamp(Yaw, -1.0f, 1.0f);
-}
-
-void UModularVehicleBaseComponent::SetHandbrakeInput(float Handbrake)
-{
-	RawHandbrakeInput = Handbrake;
+	FInputInterface Inputs(InputNameMap, RawInputsContainer);
+	Inputs.SetValue(Name, Value);
 }
 
 void UModularVehicleBaseComponent::SetReverseInput(bool Reverse)
@@ -1179,13 +1210,11 @@ void UModularVehicleBaseComponent::SetReverseInput(bool Reverse)
 	RawReverseInput = Reverse;
 }
 
-
 void UModularVehicleBaseComponent::SetGearInput(int32 Gear)
 {
 	RawGearInput = Gear;
 }
 
-
 int32 UModularVehicleBaseComponent::GetCurrentGear()
 {
 	return CurrentGear;
@@ -1214,16 +1243,14 @@ void UModularVehicleBaseComponent::ShowDebugInfo(AHUD* HUD, UCanvas* Canvas, con
 
 	// draw input values
 	Canvas->SetDrawColor(FColor::White);
-	YPos += Canvas->DrawText(RenderFont, FString::Printf(TEXT("Throttle Raw  (%3.2f) %3.2f"), RawThrottleInput, ThrottleInput), 4, YPos);
-	YPos += Canvas->DrawText(RenderFont, FString::Printf(TEXT("Boost Raw     (%3.2f) %3.2f"), RawBoostInput, BoostInput), 4, YPos);
-	YPos += Canvas->DrawText(RenderFont, FString::Printf(TEXT("Drift Raw     (%3.2f) %3.2f"), RawDriftInput, DriftInput), 4, YPos);
-	YPos += Canvas->DrawText(RenderFont, FString::Printf(TEXT("Brake Raw     (%3.2f) %3.2f"), RawBrakeInput, BrakeInput), 4, YPos);
-	YPos += Canvas->DrawText(RenderFont, FString::Printf(TEXT("Steering Raw  (%3.2f) %3.2f"), RawSteeringInput, SteeringInput), 4, YPos);
-	YPos += Canvas->DrawText(RenderFont, FString::Printf(TEXT("Handbrake Raw (%3.2f) %3.2f"), RawHandbrakeInput, HandbrakeInput), 4, YPos);
-	YPos += Canvas->DrawText(RenderFont, FString::Printf(TEXT("Roll Raw      (%3.2f) %3.2f"), RawRollInput, RollInput), 4, YPos);
-	YPos += Canvas->DrawText(RenderFont, FString::Printf(TEXT("Pitch Raw     (%3.2f) %3.2f"), RawPitchInput, PitchInput), 4, YPos);
-	YPos += Canvas->DrawText(RenderFont, FString::Printf(TEXT("Yaw Raw       (%3.2f) %3.2f"), RawYawInput, YawInput), 4, YPos);
-	YPos += Canvas->DrawText(RenderFont, FString::Printf(TEXT("Reverse Raw   (%3.2f) %3.2f"), RawReverseInput, ReverseInput), 4, YPos);
+
+	for (int I = 0; I < RawInputsContainer.GetNumInputs(); I++)
+	{
+		float Raw = RawInputsContainer.GetValueAtIndex(I).GetMagnitude();
+		float Interpolated = InputsContainer.GetValueAtIndex(I).GetMagnitude();
+
+		YPos += Canvas->DrawText(RenderFont, FString::Printf(TEXT("%s Raw  (%3.2f) %3.2f"), *InputConfig[I].Name.ToString(), Raw, Interpolated), 4, YPos);
+	}
 
 	YPos += 10;
 
@@ -1235,3 +1262,15 @@ void UModularVehicleBaseComponent::ShowDebugInfo(AHUD* HUD, UCanvas* Canvas, con
 #endif
 
 }
+
+void UModularVehicleBaseComponent::LogInputSetup()
+{
+#if !(UE_BUILD_SHIPPING || UE_BUILD_TEST)
+	const FModuleInputContainer::FInputNameMap& NameMap = InputNameMap;
+
+	for (auto& NamePair : NameMap)
+	{ 
+		UE_LOG(LogTemp, Warning, TEXT("Input: %s %d"), *NamePair.Key.ToString(), NamePair.Value);
+	}
+#endif //!(UE_BUILD_SHIPPING || UE_BUILD_TEST)
+}
\ No newline at end of file
diff --git a/Engine/Plugins/Experimental/ChaosModularVehicle/Source/ChaosModularVehicleEngine/Private/ChaosModularVehicle/ModularVehicleInputRate.cpp b/Engine/Plugins/Experimental/ChaosModularVehicle/Source/ChaosModularVehicleEngine/Private/ChaosModularVehicle/ModularVehicleInputRate.cpp
deleted file mode 100644
diff --git a/Engine/Plugins/Experimental/ChaosModularVehicle/Source/ChaosModularVehicleEngine/Private/ChaosModularVehicle/ModularVehicleSimulationCU.cpp b/Engine/Plugins/Experimental/ChaosModularVehicle/Source/ChaosModularVehicleEngine/Private/ChaosModularVehicle/ModularVehicleSimulationCU.cpp
--- a/Engine/Plugins/Experimental/ChaosModularVehicle/Source/ChaosModularVehicleEngine/Private/ChaosModularVehicle/ModularVehicleSimulationCU.cpp
+++ b/Engine/Plugins/Experimental/ChaosModularVehicle/Source/ChaosModularVehicleEngine/Private/ChaosModularVehicle/ModularVehicleSimulationCU.cpp
@@ -2,6 +2,8 @@
 
 #include "ChaosModularVehicle/ModularVehicleSimulationCU.h"
 #include "ChaosModularVehicle/ModularVehicleDefaultAsyncInput.h"
+#include "SimModule/SimModulesInclude.h"
+#include "SimModule/ModuleInput.h"
 #include "PhysicsProxy/GeometryCollectionPhysicsProxy.h"
 #include "PhysicsProxy/ClusterUnionPhysicsProxy.h"
 #include "GeometryCollection/GeometryCollectionObject.h"
@@ -28,7 +30,6 @@ FAutoConsoleVariableRef CVarChaosModularVehiclesDisableAnim(TEXT("p.ModularVehic
 void FModularVehicleSimulationCU::Initialize(TUniquePtr<Chaos::FSimModuleTree>& InSimModuleTree)
 {
 	SimModuleTree = MoveTemp(InSimModuleTree);
-	InputInterpolation.Init(FModularVehicleInputRate(), EModularVehicleInputType::Max);
 }
 
 void FModularVehicleSimulationCU::Terminate()
@@ -88,6 +89,7 @@ void FModularVehicleSimulationCU::Simulate(UWorld* InWorld, float DeltaSeconds,
 	Simulate_ClusterUnion(InWorld, DeltaSeconds, InputData, OutputData, static_cast<Chaos::FClusterUnionPhysicsProxy*>(Proxy));
 }
 
+
 void FModularVehicleSimulationCU::Simulate_ClusterUnion(UWorld* InWorld, float DeltaSeconds, const FModularVehicleAsyncInput& InputData, FModularVehicleAsyncOutput& OutputData, Chaos::FClusterUnionPhysicsProxy* Proxy)
 {
 	Chaos::EnsureIsInPhysicsThreadContext();
@@ -105,19 +107,14 @@ void FModularVehicleSimulationCU::Simulate_ClusterUnion(UWorld* InWorld, float D
 		//		, *Proxy->GetParticle_Internal()->W().ToString()));
 		//}
 
-		//InterpolateInputs(DeltaSeconds, ExternalInputs, InterpolatedInputs);
-
-		SimInputData.ControlInputs.Throttle = InputData.PhysicsInputs.NetworkInputs.VehicleInputs.Throttle;
-		SimInputData.ControlInputs.Steering = InputData.PhysicsInputs.NetworkInputs.VehicleInputs.Steering;
-		SimInputData.ControlInputs.Brake = InputData.PhysicsInputs.NetworkInputs.VehicleInputs.Brake;
-		SimInputData.ControlInputs.Handbrake = InputData.PhysicsInputs.NetworkInputs.VehicleInputs.Handbrake;
-		SimInputData.ControlInputs.Roll = InputData.PhysicsInputs.NetworkInputs.VehicleInputs.Roll;
-		SimInputData.ControlInputs.Pitch = InputData.PhysicsInputs.NetworkInputs.VehicleInputs.Pitch;
-		SimInputData.ControlInputs.Yaw = InputData.PhysicsInputs.NetworkInputs.VehicleInputs.Yaw;
-		SimInputData.ControlInputs.Boost = InputData.PhysicsInputs.NetworkInputs.VehicleInputs.Boost;
-		SimInputData.ControlInputs.Drift = InputData.PhysicsInputs.NetworkInputs.VehicleInputs.Drift;
-		SimInputData.ControlInputs.IsReversing = InputData.PhysicsInputs.NetworkInputs.VehicleInputs.Reverse;
+		FReadScopeLock InputConfigLock(InputConfigurationLock);
+
+		FModuleInputContainer Container = InputData.PhysicsInputs.NetworkInputs.VehicleInputs.Container;
+		FInputInterface InputInterface(InputNameMap, Container);
+
+		SimInputData.ControlInputs = &InputInterface;
 		SimInputData.bKeepVehicleAwake = InputData.PhysicsInputs.NetworkInputs.VehicleInputs.KeepAwake;
+		SimInputData.bIsReversing = InputData.PhysicsInputs.NetworkInputs.VehicleInputs.Reverse;
 
 		PerformAdditionalSimWork(InWorld, InputData, Proxy, SimInputData);
 
@@ -127,19 +124,6 @@ void FModularVehicleSimulationCU::Simulate_ClusterUnion(UWorld* InWorld, float D
 
 }
 
-//void FModularVehicleSimulationCU::InterpolateInputs(float DeltaSeconds, const Chaos::FControlInputs& ExternalInputIn, Chaos::FControlInputs& InterpolatedInputsInOut)
-//{
-//	InterpolatedInputsInOut.Steering = InputInterpolation[EModularVehicleInputType::Steering].InterpInputValue(DeltaSeconds, InterpolatedInputsInOut.Steering, ExternalInputIn.Steering);
-//	InterpolatedInputsInOut.Throttle = InputInterpolation[EModularVehicleInputType::Throttle].InterpInputValue(DeltaSeconds, InterpolatedInputsInOut.Throttle, ExternalInputIn.Throttle);
-//	InterpolatedInputsInOut.Brake = InputInterpolation[EModularVehicleInputType::Brake].InterpInputValue(DeltaSeconds, InterpolatedInputsInOut.Brake, ExternalInputIn.Brake);
-//	InterpolatedInputsInOut.Handbrake = InputInterpolation[EModularVehicleInputType::Handbrake].InterpInputValue(DeltaSeconds, InterpolatedInputsInOut.Handbrake, ExternalInputIn.Handbrake);
-//	InterpolatedInputsInOut.Pitch = InputInterpolation[EModularVehicleInputType::Pitch].InterpInputValue(DeltaSeconds, InterpolatedInputsInOut.Pitch, ExternalInputIn.Pitch);
-//	InterpolatedInputsInOut.Roll = InputInterpolation[EModularVehicleInputType::Roll].InterpInputValue(DeltaSeconds, InterpolatedInputsInOut.Roll, ExternalInputIn.Roll);
-//	InterpolatedInputsInOut.Yaw = InputInterpolation[EModularVehicleInputType::Yaw].InterpInputValue(DeltaSeconds, InterpolatedInputsInOut.Yaw, ExternalInputIn.Yaw);
-//	InterpolatedInputsInOut.GearNumber = ExternalInputIn.GearNumber;
-//	InterpolatedInputsInOut.InputDebugIndex = ExternalInputIn.InputDebugIndex;
-//}
-
 
 void FModularVehicleSimulationCU::PerformAdditionalSimWork(UWorld* InWorld, const FModularVehicleAsyncInput& InputData, Chaos::FClusterUnionPhysicsProxy* Proxy, Chaos::FAllInputs& AllInputs)
 {
@@ -446,8 +430,3 @@ void FModularVehicleSimulationCU::FillOutputState(FModularVehicleAsyncOutput& Ou
 	}
 }
 
-Chaos::FControlInputs& FModularVehicleSimulationCU::AccessControlInputs()
-{
-	Chaos::EnsureIsInPhysicsThreadContext();
-	return SimModuleTree->GetControlInputs();
-}
diff --git a/Engine/Plugins/Experimental/ChaosModularVehicle/Source/ChaosModularVehicleEngine/Public/ChaosModularVehicle/ChaosSimModuleManagerAsyncCallback.h b/Engine/Plugins/Experimental/ChaosModularVehicle/Source/ChaosModularVehicleEngine/Public/ChaosModularVehicle/ChaosSimModuleManagerAsyncCallback.h
--- a/Engine/Plugins/Experimental/ChaosModularVehicle/Source/ChaosModularVehicleEngine/Public/ChaosModularVehicle/ChaosSimModuleManagerAsyncCallback.h
+++ b/Engine/Plugins/Experimental/ChaosModularVehicle/Source/ChaosModularVehicleEngine/Public/ChaosModularVehicle/ChaosSimModuleManagerAsyncCallback.h
@@ -10,6 +10,7 @@
 #include "TransmissionSystem.h"
 #include "Physics/NetworkPhysicsComponent.h"
 #include "SimModule/SimulationModuleBase.h"
+#include "SimModule/ModuleInput.h"
 
 #include "ChaosSimModuleManagerAsyncCallback.generated.h"
 
@@ -46,54 +47,10 @@ struct CHAOSMODULARVEHICLEENGINE_API FModularVehicleInputs
 	GENERATED_USTRUCT_BODY()
 
 		FModularVehicleInputs()
-		: Steering(0.f)
-		, Throttle(0.f)
-		, Brake(0.f)
-		, Handbrake(0.f)
-		, Pitch(0.f)
-		, Roll(0.f)
-		, Yaw(0.f)
-		, Boost(0.f)
-		, Drift(0.f)
-		, Reverse(false)
-		, KeepAwake(false)
-		{}
-
-	// Steering output to physics system. Range -1...1
-	UPROPERTY()
-	float Steering;
-
-	// Accelerator output to physics system. Range 0...1
-	UPROPERTY()
-	float Throttle;
-
-	// Brake output to physics system. Range 0...1
-	UPROPERTY()
-	float Brake;
-
-	// Handbrake output to physics system. Range 0...1
-	UPROPERTY()
-	float Handbrake;
-
-	// Body Pitch output to physics system. Range -1...1
-	UPROPERTY()
-	float Pitch;
-
-	// Body Roll output to physics system. Range -1...1
-	UPROPERTY()
-	float Roll;
-
-	// Body Yaw output to physics system. Range -1...1
-	UPROPERTY()
-	float Yaw;
-
-	// Boost output to physics system. Range -1...1
-	UPROPERTY()
-	float Boost;
-
-	// Boost output to physics system. Range 0...1
-	UPROPERTY()
-	float Drift;
+		: Reverse(false)
+		, KeepAwake(false) 
+		{
+		}
 
 	// Reversing state
 	UPROPERTY()
@@ -103,6 +60,9 @@ struct CHAOSMODULARVEHICLEENGINE_API FModularVehicleInputs
 	UPROPERTY()
 	bool KeepAwake;
 
+	UPROPERTY()
+	FModuleInputContainer Container;
+
 };
 
 
@@ -125,8 +85,12 @@ struct CHAOSMODULARVEHICLEENGINE_API FNetworkModularVehicleInputs : public FNetw
 	/**  Serialize data function that will be used to transfer the struct across the network */
 	bool NetSerialize(FArchive& Ar, class UPackageMap* Map, bool& bOutSuccess);
 
-	/** Interpolate the data in between two input */
+	/** Interpolate the data in between two inputs */
 	virtual void InterpolateData(const FNetworkPhysicsData& MinData, const FNetworkPhysicsData& MaxData) override;
+
+	/** Merge data when multiple inputs happen an same simulation tick*/
+	virtual void MergeData(const FNetworkPhysicsData& FromData) override;
+
 };
 
 template<>
diff --git a/Engine/Plugins/Experimental/ChaosModularVehicle/Source/ChaosModularVehicleEngine/Public/ChaosModularVehicle/ModularVehicleBaseComponent.h b/Engine/Plugins/Experimental/ChaosModularVehicle/Source/ChaosModularVehicleEngine/Public/ChaosModularVehicle/ModularVehicleBaseComponent.h
--- a/Engine/Plugins/Experimental/ChaosModularVehicle/Source/ChaosModularVehicleEngine/Public/ChaosModularVehicle/ModularVehicleBaseComponent.h
+++ b/Engine/Plugins/Experimental/ChaosModularVehicle/Source/ChaosModularVehicleEngine/Public/ChaosModularVehicle/ModularVehicleBaseComponent.h
@@ -4,7 +4,7 @@
 
 #include "ChaosModularVehicle/ChaosSimModuleManagerAsyncCallback.h"
 #include "SimModule/SimModulesInclude.h"
-#include "ChaosModularVehicle/ModularVehicleInputRate.h"
+#include "SimModule/ModuleInput.h"
 #include "ChaosModularVehicle/ModularVehicleSimulationCU.h"
 #include "Chaos/ParticleHandleFwd.h"
 #include "Components/PrimitiveComponent.h"
@@ -44,99 +44,6 @@ struct CHAOSMODULARVEHICLEENGINE_API FModularReplicatedState : public FModularVe
 
 };
 
-/** Input Options */
-UENUM()
-enum class EFunctionType : uint8
-{
-	LinearFunction = 0,
-	SquaredFunction,
-	CustomCurve
-};
-
-USTRUCT()
-struct CHAOSMODULARVEHICLEENGINE_API FModularVehicleInputRateConfig
-{
-	GENERATED_USTRUCT_BODY()
-
-	/**
-		* Rate at which the input value rises
-		*/
-	UPROPERTY(EditAnywhere, Category = VehicleInputRate)
-	float RiseRate;
-
-	/**
-	 * Rate at which the input value falls
-	 */
-	UPROPERTY(EditAnywhere, Category = VehicleInputRate)
-	float FallRate;
-
-	/**
-	 * Controller input curve, various predefined options, linear, squared, or user can specify a custom curve function
-	 */
-	UPROPERTY(EditAnywhere, Category = VehicleInputRate)
-	EFunctionType InputCurveFunction;
-
-	/**
-	 * Controller input curve - should be a normalized float curve, i.e. time from 0 to 1 and values between 0 and 1
-	 * This curve is only sued if the InputCurveFunction above is set to CustomCurve
-	 */
-	UPROPERTY(EditAnywhere, Category = VehicleInputRate)
-	FRuntimeFloatCurve UserCurve;
-
-	FModularVehicleInputRateConfig() : RiseRate(5.0f), FallRate(5.0f), InputCurveFunction(EFunctionType::LinearFunction) { }
-
-	/** Change an output value using max rise and fall rates */
-	float InterpInputValue(float DeltaTime, float CurrentValue, float NewValue) const
-	{
-		const float DeltaValue = NewValue - CurrentValue;
-
-		// We are "rising" when DeltaValue has the same sign as CurrentValue (i.e. delta causes an absolute magnitude gain)
-		// OR we were at 0 before, and our delta is no longer 0.
-		const bool bRising = ((DeltaValue > 0.0f) == (CurrentValue > 0.0f)) ||
-			((DeltaValue != 0.f) && (CurrentValue == 0.f));
-
-		const float MaxDeltaValue = DeltaTime * (bRising ? RiseRate : FallRate);
-		const float ClampedDeltaValue = FMath::Clamp(DeltaValue, -MaxDeltaValue, MaxDeltaValue);
-		return CurrentValue + ClampedDeltaValue;
-	}
-
-	float CalcControlFunction(float InputValue)
-	{
-		// user defined curve
-
-		// else use option from drop down list
-		switch (InputCurveFunction)
-		{
-		case EFunctionType::CustomCurve:
-		{
-			if (UserCurve.GetRichCurveConst() && !UserCurve.GetRichCurveConst()->IsEmpty())
-			{
-				float Output = FMath::Clamp(UserCurve.GetRichCurveConst()->Eval(FMath::Abs(InputValue)), 0.0f, 1.0f);
-				return (InputValue < 0.f) ? -Output : Output;
-			}
-			else
-			{
-				return InputValue;
-			}
-		}
-		break;
-		case EFunctionType::SquaredFunction:
-		{
-			return (InputValue < 0.f) ? -InputValue * InputValue : InputValue * InputValue;
-		}
-		break;
-
-		case EFunctionType::LinearFunction:
-		default:
-		{
-			return InputValue;
-		}
-		break;
-
-		}
-
-	}
-};
 
 USTRUCT()
 struct CHAOSMODULARVEHICLEENGINE_API FConstructionData
@@ -197,6 +104,9 @@ class CHAOSMODULARVEHICLEENGINE_API UModularVehicleBaseComponent : public UPawnM
 
 	friend class FModularVehicleBuilder;
 public:
+
+	using FInputNameMap = TMap<FName, int>;
+
 	APlayerController* GetPlayerController() const;
 	bool IsLocallyControlled() const;
 
@@ -237,9 +147,6 @@ public:
 
 	FORCEINLINE const FTransform& GetComponentTransform() const;
 
-	//UPROPERTY(EditAnywhere, Category = VehicleInput, AdvancedDisplay)
-	TArray<FModularVehicleInputRate> InputInterpolationRates;
-
 	/** Use to naturally decelerate linear velocity of objects */
 	UPROPERTY(EditAnywhere, Category = "Game|Components|ModularVehicle")
 	float LinearDamping;
@@ -274,49 +181,23 @@ public:
 
 	// CONTROLS
 	// 
-	/** Set the user input for the vehicle throttle [range 0 to 1] */
-	UFUNCTION(BlueprintCallable, Category = "Game|Components|ModularVehicle")
-	void SetThrottleInput(float Throttle);
-
-	/** Set the user input for the vehicle boost [range 0 to 1] */
-	UFUNCTION(BlueprintCallable, Category = "Game|Components|ModularVehicle")
-	void SetBoostInput(float Boost);
-
-	/** Set the user input for the vehicle drift [range 0 to 1] */
-	UFUNCTION(BlueprintCallable, Category = "Game|Components|ModularVehicle")
-	void SetDriftInput(float Drift);
-
-	/** Increase the vehicle throttle position [throttle range normalized 0 to 1] */
-	UFUNCTION(BlueprintCallable, Category = "Game|Components|ModularVehicle")
-	void IncreaseThrottleInput(float ThrottleDelta);
-
-	/** Decrease the vehicle throttle position  [throttle range normalized 0 to 1] */
-	UFUNCTION(BlueprintCallable, Category = "Game|Components|ModularVehicle")
-	void DecreaseThrottleInput(float ThrottleDelta);
+	void SetInput(const FName& Name, const bool Value);
+	void SetInput(const FName& Name, const double Value);
+	void SetInput(const FName& Name, const FVector2D& Value);
+	void SetInput(const FName& Name, const FVector& Value);
 
-	/** Set the user input for the vehicle Brake [range 0 to 1] */
-	UFUNCTION(BlueprintCallable, Category = "Game|Components|ModularVehicle")
-	void SetBrakeInput(float Brake);
-
-	/** Set the user input for the vehicle steering [range -1 to 1] */
-	UFUNCTION(BlueprintCallable, Category = "Game|Components|ModularVehicle")
-	void SetSteeringInput(float Steering);
 
-	/** Set the user input for the vehicle pitch [range -1 to 1] */
 	UFUNCTION(BlueprintCallable, Category = "Game|Components|ModularVehicle")
-	void SetPitchInput(float Pitch);
+	void SetInputBool(const FName Name, const bool Value);
 
-	/** Set the user input for the vehicle roll [range -1 to 1] */
 	UFUNCTION(BlueprintCallable, Category = "Game|Components|ModularVehicle")
-	void SetRollInput(float Roll);
+	void SetInputAxis1D(const FName Name, const double Value);
 
-	/** Set the user input for the vehicle yaw [range -1 to 1] */
 	UFUNCTION(BlueprintCallable, Category = "Game|Components|ModularVehicle")
-	void SetYawInput(float Yaw);
+	void SetInputAxis2D(const FName Name, const FVector2D Value);
 
-	/** Set the user input for handbrake */
 	UFUNCTION(BlueprintCallable, Category = "Game|Components|ModularVehicle")
-	void SetHandbrakeInput(float Handbrake);
+	void SetInputAxis3D(const FName Name, const FVector Value);
 
 	UFUNCTION(BlueprintCallable, Category = "Game|Components|ModularVehicle")
 	void SetReverseInput(bool Reverse);
@@ -380,6 +261,10 @@ protected:
 	void AddOverlappingComponentsToCluster();
 	void AddGeometryCollectionsFromOwnedActor();
 	void SetupSkeletalAnimationStructure();
+	void AssimilateComponentInputs(TArray<FModuleInputSetup>& OutCombinedInputs);
+	virtual void GenerateInputModifiers(const TArray<FModuleInputSetup>& CombinedInputConfiguration);
+	virtual void ApplyInputModifiers(float DeltaTime, const FModuleInputContainer& RawValue);
+
 
 	void ActionTreeUpdates(Chaos::FSimTreeUpdates* NextTreeUpdates);
 
@@ -388,6 +273,7 @@ protected:
 
 	IPhysicsProxyBase* GetPhysicsProxy() const;
 
+
 	int32 FindComponentAddOrder(UPrimitiveComponent* InComponent);
 	bool FindAndRemoveNextPendingUpdate(int32 NextIndex, Chaos::FSimTreeUpdates* OutData);
 
@@ -395,46 +281,6 @@ protected:
 	UPROPERTY(Transient, Replicated)
 	FModularReplicatedState ReplicatedState;
 
-	// What the player has the steering set to. Range -1...1
-	UPROPERTY(Transient)
-	float RawSteeringInput;
-
-	// What the player has the accelerator set to. Range -1...1
-	UPROPERTY(Transient)
-	float RawThrottleInput;
-
-	// What the player has the brake set to. Range -1...1
-	UPROPERTY(Transient)
-	float RawBrakeInput;
-
-	// What the player has the brake set to. Range -1...1
-	UPROPERTY(Transient)
-	float RawHandbrakeInput;
-
-	// What the player has the clutch set to. Range -1...1
-	UPROPERTY(Transient)
-	float RawClutchInput;
-
-	// What the player has the pitch set to. Range -1...1
-	UPROPERTY(Transient)
-	float RawPitchInput;
-
-	// What the player has the roll set to. Range -1...1
-	UPROPERTY(Transient)
-	float RawRollInput;
-
-	// What the player has the yaw set to. Range -1...1
-	UPROPERTY(Transient)
-	float RawYawInput;
-
-	// What the player has the yaw set to. Range -1...1
-	UPROPERTY(Transient)
-	float RawBoostInput;
-
-	// What the player has the yaw set to. Range -1...1
-	UPROPERTY(Transient)
-	float RawDriftInput;
-
 	// latest gear selected
 	UPROPERTY(Transient)
 	int32 RawGearInput;
@@ -443,46 +289,6 @@ protected:
 	UPROPERTY(Transient)
 	bool RawReverseInput;
 
-	// Steering output to physics system. Range -1...1
-	UPROPERTY(Transient)
-	float SteeringInput;
-
-	// Accelerator output to physics system. Range 0...1
-	UPROPERTY(Transient)
-	float ThrottleInput;
-
-	// Brake output to physics system. Range 0...1
-	UPROPERTY(Transient)
-	float BrakeInput;
-
-	// Handbrake output to physics system. Range 0...1
-	UPROPERTY(Transient)
-	float HandbrakeInput;
-
-	// Clutch output to physics system. Range 0...1
-	UPROPERTY(Transient)
-	float ClutchInput;
-
-	// Body Pitch output to physics system. Range -1...1
-	UPROPERTY(Transient)
-	float PitchInput;
-
-	// Body Roll output to physics system. Range -1...1
-	UPROPERTY(Transient)
-	float RollInput;
-
-	// Body Yaw output to physics system. Range -1...1
-	UPROPERTY(Transient)
-	float YawInput;
-
-	// Boost output to physics system. Range 0...1
-	UPROPERTY(Transient)
-	float BoostInput;
-
-	// Boost output to physics system. Range 0...1
-	UPROPERTY(Transient)
-	float DriftInput;
-
 	// Reverse state
 	UPROPERTY(Transient)
 	bool ReverseInput;
@@ -504,50 +310,20 @@ protected:
 
 public:
 
-	// Rate at which input throttle can rise and fall
-	UPROPERTY(EditAnywhere, Category = VehicleInput, AdvancedDisplay)
-	FModularVehicleInputRateConfig ThrottleInputRate;
-
-	// Rate at which input brake can rise and fall
-	UPROPERTY(EditAnywhere, Category = VehicleInput, AdvancedDisplay)
-	FModularVehicleInputRateConfig BrakeInputRate;
-
-	// Rate at which input steering can rise and fall
-	UPROPERTY(EditAnywhere, Category = VehicleInput, AdvancedDisplay)
-	FModularVehicleInputRateConfig SteeringInputRate;
-
-	// Rate at which input handbrake can rise and fall
-	UPROPERTY(EditAnywhere, Category = VehicleInput, AdvancedDisplay)
-	FModularVehicleInputRateConfig HandbrakeInputRate;
-
-	// Rate at which input pitch can rise and fall
-	UPROPERTY(EditAnywhere, Category = VehicleInput, AdvancedDisplay)
-	FModularVehicleInputRateConfig PitchInputRate;
-
-	// Rate at which input roll can rise and fall
-	UPROPERTY(EditAnywhere, Category = VehicleInput, AdvancedDisplay)
-	FModularVehicleInputRateConfig RollInputRate;
-
-	// Rate at which input yaw can rise and fall
-	UPROPERTY(EditAnywhere, Category = VehicleInput, AdvancedDisplay)
-	FModularVehicleInputRateConfig YawInputRate;
-
-	// Rate at which input can rise and fall
-	UPROPERTY(EditAnywhere, Category = VehicleInput, AdvancedDisplay)
-	FModularVehicleInputRateConfig BoostInputRate;
-
-	// Rate at which input can rise and fall
-	UPROPERTY(EditAnywhere, Category = VehicleInput, AdvancedDisplay)
-	FModularVehicleInputRateConfig DriftInputRate;
+	UPROPERTY(EditAnywhere, Category = VehicleInput)
+	TArray<FModuleInputSetup> InputConfig;
 
 	UPROPERTY(Transient, Replicated)
 	TArray<FConstructionData> ConstructionDatas;
 
+	UPROPERTY(transient, duplicatetransient, BlueprintReadOnly, Category = "Game|Components|ModularVehicle")
+	TArray<TObjectPtr<UDefaultModularVehicleInputModifier>> InputModifiers;
+
 	/** Pass current state to server */
 	UFUNCTION(reliable, server, WithValidation)
-	void ServerUpdateState(float InSteeringInput, float InThrottleInput, float InBrakeInput
-		, float InHandbrakeInput, int32 InCurrentGear, float InRollInput, float InPitchInput
-		, float InYawInput, float InBoostInput, float InDriftInput, bool InReverseInput);
+	void ServerUpdateState(const FModuleInputContainer& InputsIn, bool Reverse, bool KeepAwake);
+
+	void LogInputSetup();
 
 	TArray<AActor*> ActorsToIgnore;
 	EChaosAsyncVehicleDataType CurAsyncType;
@@ -592,5 +368,9 @@ private:
 
 	TArray<FModuleAnimationSetup> ModuleAnimationSetups;
 
+	FInputNameMap InputNameMap;	// map input name to input container array index
+	FModuleInputContainer RawInputsContainer;
+	FModuleInputContainer InputsContainer;
+
 };
 
diff --git a/Engine/Plugins/Experimental/ChaosModularVehicle/Source/ChaosModularVehicleEngine/Public/ChaosModularVehicle/ModularVehicleInputRate.h b/Engine/Plugins/Experimental/ChaosModularVehicle/Source/ChaosModularVehicleEngine/Public/ChaosModularVehicle/ModularVehicleInputRate.h
deleted file mode 100644
diff --git a/Engine/Plugins/Experimental/ChaosModularVehicle/Source/ChaosModularVehicleEngine/Public/ChaosModularVehicle/ModularVehicleSimulationCU.h b/Engine/Plugins/Experimental/ChaosModularVehicle/Source/ChaosModularVehicleEngine/Public/ChaosModularVehicle/ModularVehicleSimulationCU.h
--- a/Engine/Plugins/Experimental/ChaosModularVehicle/Source/ChaosModularVehicleEngine/Public/ChaosModularVehicle/ModularVehicleSimulationCU.h
+++ b/Engine/Plugins/Experimental/ChaosModularVehicle/Source/ChaosModularVehicleEngine/Public/ChaosModularVehicle/ModularVehicleSimulationCU.h
@@ -3,7 +3,6 @@
 #pragma once
 
 #include "GeometryCollection/GeometryCollectionComponent.h"
-#include "ChaosModularVehicle/ModularVehicleInputRate.h"
 #include "ChaosModularVehicle/ChaosSimModuleManagerAsyncCallback.h"
 #include "SimModule/SimModuleTree.h"
 #include "SimModule/SimModulesInclude.h"
@@ -11,6 +10,20 @@
 
 struct FModularVehicleAsyncInput;
 struct FChaosSimModuleManagerAsyncOutput;
+struct FModuleInputContainer;
+
+struct CHAOSMODULARVEHICLEENGINE_API FModularVehicleDebugParams
+{
+	bool ShowDebug = false;
+	bool SuspensionRaycastsEnabled = true;
+	bool ShowSuspensionRaycasts = false;
+	bool ShowWheelData = false;
+	bool ShowRaycastMaterial = false;
+	bool ShowWheelCollisionNormal = false;
+
+	bool DisableAnim = false;
+	float FrictionOverride = 1.0f;
+};
 
 namespace Chaos
 {
@@ -21,6 +34,8 @@ namespace Chaos
 class CHAOSMODULARVEHICLEENGINE_API FModularVehicleSimulationCU
 {
 public:
+	using FInputNameMap = FInputInterface::FInputNameMap;
+
 	FModularVehicleSimulationCU(bool InUsingNetworkPhysicsPrediction, int8 InNetMode)
 		: bUsingNetworkPhysicsPrediction(InUsingNetworkPhysicsPrediction)
 		, NetMode(InNetMode)
@@ -34,6 +49,11 @@ public:
 
 	void Initialize(TUniquePtr<Chaos::FSimModuleTree>& InSimModuleTree);
 	void Terminate();
+	void SetInputMappings(const FInputNameMap& InNameMap)
+	{ 
+		FWriteScopeLock InputConfigLock(InputConfigurationLock);
+		InputNameMap = InNameMap; 
+	}
 
 	/** Update called from Physics Thread */
 	virtual void Simulate(UWorld* InWorld, float DeltaSeconds, const FModularVehicleAsyncInput& InputData, FModularVehicleAsyncOutput& OutputData, IPhysicsProxyBase* Proxy);
@@ -46,8 +66,6 @@ public:
 
 	void FillOutputState(FModularVehicleAsyncOutput& Output);
 
-	Chaos::FControlInputs& AccessControlInputs();
-
 	const TUniquePtr<Chaos::FSimModuleTree>& GetSimComponentTree() const {
 		Chaos::EnsureIsInPhysicsThreadContext();
 		return SimModuleTree;
@@ -57,16 +75,14 @@ public:
 		return SimModuleTree; 
 		}
 
-	TArray<FModularVehicleInputRate>& AccessInputInterpolation() { return InputInterpolation; }
-
-
 	TUniquePtr<Chaos::FSimModuleTree> SimModuleTree;	/* Simulation modules stored in tree structure */
-	TArray<FModularVehicleInputRate> InputInterpolation;
 	Chaos::FAllInputs SimInputData;
 	bool bUsingNetworkPhysicsPrediction;
 
 	/** Current control inputs that is being used on the PT */
 	FModularVehicleInputs VehicleInputs;
+	FInputNameMap InputNameMap;
+	FRWLock InputConfigurationLock;
 
 	int8 NetMode;
 
diff --git a/Engine/Plugins/Experimental/ChaosModularVehicle/Source/ChaosModularVehicleEngine/Public/ChaosModularVehicle/VehicleSimBaseComponent.h b/Engine/Plugins/Experimental/ChaosModularVehicle/Source/ChaosModularVehicleEngine/Public/ChaosModularVehicle/VehicleSimBaseComponent.h
--- a/Engine/Plugins/Experimental/ChaosModularVehicle/Source/ChaosModularVehicleEngine/Public/ChaosModularVehicle/VehicleSimBaseComponent.h
+++ b/Engine/Plugins/Experimental/ChaosModularVehicle/Source/ChaosModularVehicleEngine/Public/ChaosModularVehicle/VehicleSimBaseComponent.h
@@ -6,6 +6,7 @@
 #include "UObject/ObjectPtr.h"
 #include "Components/PrimitiveComponent.h"
 #include "Components/SceneComponent.h"
+#include "SimModule/ModuleInput.h"
 
 #include "VehicleSimBaseComponent.generated.h"
 
@@ -58,6 +59,9 @@ public:
 	UPROPERTY(EditAnywhere, BlueprintReadWrite, Category = ModularVehicle)
 	bool bAnimationEnabled;
 
+	UPROPERTY(EditAnywhere, Category = VehicleInput)
+	TArray<FModuleInputSetup> InputConfig;
+
 	int AnimationSetupIndex;
 	int TreeIndex; // helper - since Component->GetAttachChildren doesn't contain any data
 };
diff --git a/Engine/Source/Programs/HeadlessChaos/Private/HeadlessChaosTestModularVehicles.cpp b/Engine/Source/Programs/HeadlessChaos/Private/HeadlessChaosTestModularVehicles.cpp
--- a/Engine/Source/Programs/HeadlessChaos/Private/HeadlessChaosTestModularVehicles.cpp
+++ b/Engine/Source/Programs/HeadlessChaos/Private/HeadlessChaosTestModularVehicles.cpp
@@ -12,6 +12,7 @@
 #include "SimModule/WheelModule.h"
 #include "SimModule/SuspensionModule.h"
 #include "SimModule/SimModuleTree.h"
+#include "SimModule/ModuleInput.h"
 
 // for Simulation Tests
 #include "Chaos/PBDRigidsEvolutionGBF.h"
@@ -29,6 +30,57 @@ namespace ChaosTest
 {
 	using namespace Chaos;
 
+	struct FInputsContainer
+	{
+		FInputsContainer(Chaos::FAllInputs& InInputs)
+		{
+			ConfigureControlInputs(InInputs);
+		}
+		
+		void ConfigureControlInputs(Chaos::FAllInputs& Inputs)
+		{
+			FModuleInputSetup InputData1(TEXT("Throttle"), EModuleInputValueType::MAxis1D);
+			InputSetupData.Add(InputData1);
+
+			FModuleInputSetup InputData2(TEXT("Brake"), EModuleInputValueType::MAxis1D);
+			InputSetupData.Add(InputData2);
+
+			FModuleInputSetup InputData3(TEXT("Steering"), EModuleInputValueType::MAxis1D);
+			InputSetupData.Add(InputData3);
+
+			FModuleInputSetup InputData4(TEXT("Clutch"), EModuleInputValueType::MAxis1D);
+			InputSetupData.Add(InputData4);
+
+			FModuleInputSetup InputData5(TEXT("Handbrake"), EModuleInputValueType::MAxis1D);
+			InputSetupData.Add(InputData5);
+
+			FModuleInputSetup InputData6(TEXT("Pitch"), EModuleInputValueType::MAxis1D);
+			InputSetupData.Add(InputData6);
+
+			FModuleInputSetup InputData7(TEXT("Yaw"), EModuleInputValueType::MAxis1D);
+			InputSetupData.Add(InputData7);
+
+			FModuleInputSetup InputData8(TEXT("Roll"), EModuleInputValueType::MAxis1D);
+			InputSetupData.Add(InputData8);
+
+			FModuleInputSetup InputData9(TEXT("ChangeUp"), EModuleInputValueType::MBoolean);
+			InputSetupData.Add(InputData9);
+
+			FModuleInputSetup InputData10(TEXT("ChangeDown"), EModuleInputValueType::MBoolean);
+			InputSetupData.Add(InputData10);
+
+			ValueContainer.Initialize(InputSetupData, NameMap);
+
+			ControlInputs = MakeUnique<FInputInterface>(NameMap, ValueContainer);
+			Inputs.ControlInputs = ControlInputs.Get();
+		}
+
+		FInputInterface::FInputNameMap NameMap;
+		FModuleInputContainer ValueContainer;
+		TArray<FModuleInputSetup> InputSetupData;
+		TUniquePtr<FInputInterface> ControlInputs;
+	};
+
 	GTEST_TEST(AllTraits, ModularVehicleTest_Aerofoil)
 	{
 		FAerofoilSettings RWingSetup;
@@ -153,6 +205,7 @@ namespace ChaosTest
 		void Test_TransmissionManualGearSelection()
 		{
 			FAllInputs Inputs;
+			FInputsContainer IC(Inputs);
 			FSimModuleTree Tree;
 
 			EXPECT_EQ(GetCurrentGear(), 1);
@@ -215,6 +268,7 @@ namespace ChaosTest
 		void Test_TransmissionAutoGearSelection()
 		{
 			FAllInputs Inputs;
+			FInputsContainer IC(Inputs);
 			FSimModuleTree Tree;
 
 			SetGear(1, true);
@@ -376,8 +430,9 @@ namespace ChaosTest
 		)
 	{
 		FAllInputs Inputs;
-		Inputs.ControlInputs.Brake = 1.0f; // apply full braking force
-		Inputs.ControlInputs.Throttle = 0.0f;
+		FInputsContainer IC(Inputs);
+		Inputs.GetControls().SetValue(TEXT("Throttle"), 0.0f);
+		Inputs.GetControls().SetValue(TEXT("Brake"), 1.0f); // Apply full brake
 		FSimModuleTree Tree;
 
 		StoppingDistanceOut = 0.f;
@@ -442,8 +497,9 @@ namespace ChaosTest
 		)
 	{
 		FAllInputs Inputs;
-		Inputs.ControlInputs.Throttle = 1.0f; // apply full throttle
-		Inputs.ControlInputs.Brake = 0.0f;
+		FInputsContainer IC(Inputs);
+		Inputs.GetControls().SetValue(TEXT("Throttle"), 1.0f); // Apply full throttle
+		Inputs.GetControls().SetValue(TEXT("Brake"), 0.0f);
 		FSimModuleTree Tree;
 
 		DistanceTravelledOut = 0.f;
@@ -666,8 +722,9 @@ namespace ChaosTest
 		FWheelSimModule Wheel(Setup);
 
 		FAllInputs Inputs;
-		Inputs.ControlInputs.Throttle = 0.0f;
-		Inputs.ControlInputs.Brake = 0.0f;
+		FInputsContainer IC(Inputs);
+		Inputs.GetControls().SetValue(TEXT("Throttle"), 0.0f);
+		Inputs.GetControls().SetValue(TEXT("Brake"), 0.0f);
 		FSimModuleTree Tree;
 
 		float DeltaTime = 1.f / 30.f;
@@ -756,6 +813,7 @@ namespace ChaosTest
 	GTEST_TEST(AllTraits, ModularVehicleTest_EngineRPM)
 	{
 		FAllInputs Inputs;
+		FInputsContainer IC(Inputs);
 		FSimModuleTree Tree;
 
 		FEngineSettings Setup;
@@ -784,7 +842,7 @@ namespace ChaosTest
 
 		FEngineSimModule Engine(Setup);
 
-		Inputs.ControlInputs.Throttle = 0.f;
+		Inputs.GetControls().SetValue(TEXT("Throttle"), 0.0f);
 
 		float DeltaTime = 1.0f / 30.0f;
 		float TOLERANCE = 0.1f;
@@ -798,7 +856,7 @@ namespace ChaosTest
 		EXPECT_LT(Engine.GetRPM() - Engine.Setup().IdleRPM, TOLERANCE);
 
 		// apply half throttle
-		Inputs.ControlInputs.Throttle = 0.5f;
+		Inputs.GetControls().SetValue(TEXT("Throttle"), 0.5f);
 
 		for (int i = 0; i < 100; i++)
 		{
@@ -808,7 +866,7 @@ namespace ChaosTest
 
 		EXPECT_GT(Engine.GetRPM(), Engine.Setup().IdleRPM);
 
-		Inputs.ControlInputs.Throttle = 0.0f;
+		Inputs.GetControls().SetValue(TEXT("Throttle"), 0.0f);
 
 		// engine idle - no throttle
 		for (int i = 0; i < 200; i++)
@@ -825,6 +883,7 @@ namespace ChaosTest
 	GTEST_TEST(AllTraits, ModularVehicleTest_SimModuleTree_EngineDrivingWheelsThroughClutch)
 	{
 		FAllInputs Inputs;
+		FInputsContainer IC(Inputs);
 		FSimModuleTree Tree;
 		float TOLERANCE = 0.1f;
 
@@ -887,9 +946,9 @@ namespace ChaosTest
 		}
 
 		// Throttle ON
-		Inputs.ControlInputs.Throttle = 1.0f;
-		Inputs.ControlInputs.Brake = 0.0f;
-		Inputs.ControlInputs.Clutch = 0.0f;
+		Inputs.GetControls().SetValue(TEXT("Throttle"), 1.0f);
+		Inputs.GetControls().SetValue(TEXT("Brake"), 0.0f);
+		Inputs.GetControls().SetValue(TEXT("Clutch"), 0.0f);
 
 		// wheels not moving initially
 		EXPECT_NEAR(Wheel0.GetRPM(), 0, TOLERANCE);
@@ -912,9 +971,9 @@ namespace ChaosTest
 		EXPECT_NEAR(Wheel3.GetRPM(), 0, TOLERANCE);
 
 		// Brake ON 
-		Inputs.ControlInputs.Throttle = 0.0f;
-		Inputs.ControlInputs.Brake = 1.0f;
-		Inputs.ControlInputs.Clutch = 0.0f;
+		Inputs.GetControls().SetValue(TEXT("Throttle"), 0.0f);
+		Inputs.GetControls().SetValue(TEXT("Brake"), 1.0f);
+		Inputs.GetControls().SetValue(TEXT("Clutch"), 0.0f);
 
 		// Simulate
 		for (int I = 0; I < 50; I++)
@@ -930,9 +989,9 @@ namespace ChaosTest
 
 
 		// Throttle ON, Clutch depressed
-		Inputs.ControlInputs.Throttle = 1.0f;
-		Inputs.ControlInputs.Brake = 0.0f;
-		Inputs.ControlInputs.Clutch = 1.0f;
+		Inputs.GetControls().SetValue(TEXT("Throttle"), 1.0f);
+		Inputs.GetControls().SetValue(TEXT("Brake"), 0.0f);
+		Inputs.GetControls().SetValue(TEXT("Clutch"), 1.0f);
 
 		// Simulate
 		for (int I = 0; I < 50; I++)
@@ -949,14 +1008,15 @@ namespace ChaosTest
 	GTEST_TEST(AllTraits, ModularVehicleTest_SimModuleTree_WheelsSpinEngineShouldSpin)
 	{
 		FAllInputs Inputs;
+		FInputsContainer IC(Inputs);
 		FSimModuleTree Tree;
 		float TOLERANCE = 0.1f;
 
 		FChassisSettings ChassisSettings;
 		int RootNodeIndex = Tree.AddRoot(new FChassisSimModule(ChassisSettings));
 
-		Inputs.ControlInputs.Brake = 0;
-		Inputs.ControlInputs.Throttle = 0;
+		Inputs.GetControls().SetValue(TEXT("Throttle"), 0.0f);
+		Inputs.GetControls().SetValue(TEXT("Brake"), 0.0f);
 
 		FEngineSettings EngineSettings;
 		{
diff --git a/Engine/Source/Runtime/Experimental/ChaosVehicles/ChaosVehiclesCore/Private/SimModule/AerofoilModule.cpp b/Engine/Source/Runtime/Experimental/ChaosVehicles/ChaosVehiclesCore/Private/SimModule/AerofoilModule.cpp
--- a/Engine/Source/Runtime/Experimental/ChaosVehicles/ChaosVehiclesCore/Private/SimModule/AerofoilModule.cpp
+++ b/Engine/Source/Runtime/Experimental/ChaosVehicles/ChaosVehiclesCore/Private/SimModule/AerofoilModule.cpp
@@ -10,7 +10,6 @@ UE_DISABLE_OPTIMIZATION_SHIP
 
 namespace Chaos
 {
-
 	FAerofoilSimModule::FAerofoilSimModule(const FAerofoilSettings& Settings) : TSimModuleSettings<FAerofoilSettings>(Settings)
 		, CurrentAirDensity(RealWorldConsts::AirDensity())
 		, AngleOfAttack(0.f)
@@ -27,15 +26,15 @@ namespace Chaos
 		switch (Setup().Type)
 		{
 			case EAerofoil::Elevator:
-				ControlSurfaceAngle = Inputs.ControlInputs.Pitch * Setup().MaxControlAngle;
+				ControlSurfaceAngle = Inputs.GetControls().GetMagnitude(PitchControlName) * Setup().MaxControlAngle;
 				break;
 
 			case EAerofoil::Rudder:
-				ControlSurfaceAngle = Inputs.ControlInputs.Yaw * Setup().MaxControlAngle;
+				ControlSurfaceAngle = Inputs.GetControls().GetMagnitude(YawControlName) * Setup().MaxControlAngle;
 				break;
 
 			case EAerofoil::Wing:
-				ControlSurfaceAngle = Inputs.ControlInputs.Roll * Setup().MaxControlAngle;
+				ControlSurfaceAngle = Inputs.GetControls().GetMagnitude(RollControlName) * Setup().MaxControlAngle;
 			break;
 		}
 
diff --git a/Engine/Source/Runtime/Experimental/ChaosVehicles/ChaosVehiclesCore/Private/SimModule/ClutchModule.cpp b/Engine/Source/Runtime/Experimental/ChaosVehicles/ChaosVehiclesCore/Private/SimModule/ClutchModule.cpp
--- a/Engine/Source/Runtime/Experimental/ChaosVehicles/ChaosVehiclesCore/Private/SimModule/ClutchModule.cpp
+++ b/Engine/Source/Runtime/Experimental/ChaosVehicles/ChaosVehiclesCore/Private/SimModule/ClutchModule.cpp
@@ -10,7 +10,6 @@ UE_DISABLE_OPTIMIZATION_SHIP
 
 namespace Chaos
 {
-
 	FClutchSimModule::FClutchSimModule(const FClutchSettings& Settings) : TSimModuleSettings<FClutchSettings>(Settings)
 		, ClutchValue(0.f)
 	{
@@ -33,7 +32,7 @@ namespace Chaos
 		//float AngularVelocityDifference = EngineSpeed - TransmissionSpeed;
 
 		// Inputs.Clutch 0 is engaged/locked, 1 is depressed/open
-		ClutchValue = (1.0f - Inputs.ControlInputs.Clutch) * Setup().ClutchStrength;
+		ClutchValue = (1.0f - Inputs.GetControls().GetMagnitude(ClutchControlName)) * Setup().ClutchStrength;
 
 		FTorqueSimModule* Parent = static_cast<FTorqueSimModule*>(GetParent());
 		FTorqueSimModule* Child = static_cast<FTorqueSimModule*>(GetFirstChild());
diff --git a/Engine/Source/Runtime/Experimental/ChaosVehicles/ChaosVehiclesCore/Private/SimModule/EngineModule.cpp b/Engine/Source/Runtime/Experimental/ChaosVehicles/ChaosVehiclesCore/Private/SimModule/EngineModule.cpp
--- a/Engine/Source/Runtime/Experimental/ChaosVehicles/ChaosVehiclesCore/Private/SimModule/EngineModule.cpp
+++ b/Engine/Source/Runtime/Experimental/ChaosVehicles/ChaosVehiclesCore/Private/SimModule/EngineModule.cpp
@@ -10,7 +10,6 @@ UE_DISABLE_OPTIMIZATION_SHIP
 
 namespace Chaos
 {
-
 	FEngineSimModule::FEngineSimModule(const FEngineSettings& Settings) : TSimModuleSettings<FEngineSettings>(Settings)
 		, EngineIdleSpeed(RPMToOmega(Setup().IdleRPM))
 		, MaxEngineSpeed(RPMToOmega(Setup().MaxRPM))
@@ -32,7 +31,7 @@ namespace Chaos
 		}
 
 		// TODO: Engine braking effect
-		DriveTorque = GetEngineTorque(Inputs.ControlInputs.Throttle, GetRPM());
+		DriveTorque = GetEngineTorque(Inputs.GetControls().GetMagnitude(ThrottleControlName), GetRPM());
 
 		if (DriveTorque < SMALL_NUMBER)
 		{
diff --git a/Engine/Source/Runtime/Experimental/ChaosVehicles/ChaosVehiclesCore/Private/SimModule/ModuleInput.cpp b/Engine/Source/Runtime/Experimental/ChaosVehicles/ChaosVehiclesCore/Private/SimModule/ModuleInput.cpp
new file mode 100644
--- /dev/null
+++ b/Engine/Source/Runtime/Experimental/ChaosVehicles/ChaosVehiclesCore/Private/SimModule/ModuleInput.cpp
@@ -0,0 +1,307 @@
+// Copyright Epic Games, Inc. All Rights Reserved.
+
+#include "SimModule/ModuleInput.h"
+
+TArray<struct FModuleInputSetup>* FScopedModuleInputInitializer::InitSetupData = nullptr;
+
+DEFINE_LOG_CATEGORY(LogModularInput);
+
+float FModuleInputValue::GetMagnitudeSq() const
+{
+	switch (GetValueType())
+	{
+	case EModuleInputValueType::MBoolean:
+	case EModuleInputValueType::MAxis1D:
+		return Value.X * Value.X;
+	case EModuleInputValueType::MAxis2D:
+		return Value.SizeSquared2D();
+	case EModuleInputValueType::MAxis3D:
+		return Value.SizeSquared();
+	}
+
+	checkf(false, TEXT("Unsupported value type for module input value!"));
+	return 0.f;
+}
+
+float FModuleInputValue::GetMagnitude() const
+{
+	switch (GetValueType())
+	{
+	case EModuleInputValueType::MBoolean:
+	case EModuleInputValueType::MAxis1D:
+		return Value.X;
+	case EModuleInputValueType::MAxis2D:
+		return Value.Size2D();
+	case EModuleInputValueType::MAxis3D:
+		return Value.Size();
+	}
+
+	checkf(false, TEXT("Unsupported value type for module input value!"));
+	return 0.f;
+}
+
+void FModuleInputValue::SetMagnitude(float NewSize)
+{
+	switch (GetValueType())
+	{
+	case EModuleInputValueType::MBoolean:
+	case EModuleInputValueType::MAxis1D:
+		Value.X = NewSize;
+		break;
+	case EModuleInputValueType::MAxis2D:
+		Value.GetSafeNormal2D()* NewSize;
+		break;
+	case EModuleInputValueType::MAxis3D:
+		Value.GetSafeNormal()* NewSize;
+		break;
+	default:
+		checkf(false, TEXT("Unsupported value type for module input value!"));
+		break;
+	}
+}
+
+void FModuleInputValue::Serialize(FArchive& Ar, UPackageMap* Map, bool& bOutSuccess)
+{
+	// #TODO: send only when value changes also quantize what is sent, but just get it working for now
+	Ar << ValueType;
+
+	switch (GetValueType())
+	{
+	case EModuleInputValueType::MBoolean:
+		Ar << Value.X;
+		break;
+	case EModuleInputValueType::MAxis1D:
+		Ar << Value.X;
+		break;
+	case EModuleInputValueType::MAxis2D:
+		Ar << Value.X;
+		Ar << Value.Y;
+		break;
+	case EModuleInputValueType::MAxis3D:
+		Ar << Value.X;
+		Ar << Value.Y;
+		Ar << Value.Z;
+		break;
+	default:
+		checkf(false, TEXT("Unsupported value type for module input value!"));
+		break;
+	}
+	bOutSuccess = true;
+}
+
+void FModuleInputValue::Merge(const FModuleInputValue& From)
+{
+	switch (GetValueType())
+	{
+	case EModuleInputValueType::MBoolean:
+		// ensure we capture edges of digitial inputs by taking largest absolute value
+		if (FMath::Abs(From.Value.X) >= FMath::Abs(Value.X))
+		{
+			Value.X = From.Value.X;
+		}
+		break;
+	case EModuleInputValueType::MAxis1D:
+	case EModuleInputValueType::MAxis2D:
+	case EModuleInputValueType::MAxis3D:
+		// use the last known value for analog inputs, or should we take the highest value of the pair here also?
+		Value = From.Value;			
+		break;
+	default:
+		checkf(false, TEXT("Unsupported value type for module input value!"));
+		break;
+	}
+
+}
+
+FString FModuleInputValue::ToString() const
+{
+	switch (GetValueType())
+	{
+	case EModuleInputValueType::MBoolean:
+		return FString(IsNonZero() ? TEXT("true") : TEXT("false"));
+	case EModuleInputValueType::MAxis1D:
+		return FString::Printf(TEXT("%3.3f"), Value.X);
+	case EModuleInputValueType::MAxis2D:
+		return FString::Printf(TEXT("X=%3.3f Y=%3.3f"), Value.X, Value.Y);
+	case EModuleInputValueType::MAxis3D:
+		return FString::Printf(TEXT("X=%3.3f Y=%3.3f Z=%3.3f"), Value.X, Value.Y, Value.Z);
+	}
+
+	checkf(false, TEXT("Unsupported value type for module input value!"));
+	return FString{};
+}
+
+
+
+FModuleInputValue UDefaultModularVehicleInputModifier::InterpInputValue(float DeltaTime, const FModuleInputValue& CurrentValue, const FModuleInputValue& NewValue) const
+{
+	const FModuleInputValue DeltaValue = NewValue - CurrentValue;
+
+	// We are "rising" when DeltaValue has the same sign as CurrentValue (i.e. delta causes an absolute magnitude gain)
+	// OR we were at 0 before, and our delta is no longer 0.
+	const bool bRising = ((DeltaValue.GetMagnitude() > 0.0f) == (CurrentValue.GetMagnitude() > 0.0f)) ||
+		((DeltaValue.GetMagnitude() != 0.f) && (CurrentValue.GetMagnitude() == 0.f));
+
+	const float MaxMagnitude = DeltaTime * (bRising ? RiseRate : FallRate);
+
+	const FModuleInputValue ClampedDeltaValue = FModuleInputValue::Clamp(DeltaValue, -MaxMagnitude, MaxMagnitude);
+
+	return CurrentValue + ClampedDeltaValue;
+}
+
+float UDefaultModularVehicleInputModifier::CalcControlFunction(float InputValue)
+{
+	// #TODO: Reinstate ?
+	check(false);
+	return 0.0f;
+	// user defined curve
+
+	//// else use option from drop down list
+	//switch (InputCurveFunction)
+	//{
+	//case EFunctionType::CustomCurve:
+	//{
+	//	if (UserCurve.GetRichCurveConst() && !UserCurve.GetRichCurveConst()->IsEmpty())
+	//	{
+	//		float Output = FMath::Clamp(UserCurve.GetRichCurveConst()->Eval(FMath::Abs(InputValue)), 0.0f, 1.0f);
+	//		return (InputValue < 0.f) ? -Output : Output;
+	//	}
+	//	else
+	//	{
+	//		return InputValue;
+	//	}
+	//}
+	//break;
+	//case EFunctionType::SquaredFunction:
+	//{
+	//	return (InputValue < 0.f) ? -InputValue * InputValue : InputValue * InputValue;
+	//}
+	//break;
+
+	//case EFunctionType::LinearFunction:
+	//default:
+	//{
+	//	return InputValue;
+	//}
+	//break;
+
+	//}
+
+}
+
+
+void FModuleInputContainer::Initialize(TArray<FModuleInputSetup>& SetupData, FInputNameMap& NameMapOut)
+{
+	NameMapOut.Reset();
+	InputValues.Reset();
+
+	for (FModuleInputSetup& Setup : SetupData)
+	{
+		int Index = AddInput(Setup.Type, Setup.InputModifierClass);
+		NameMapOut.Add(Setup.Name, Index);
+	}
+}
+
+void FModuleInputContainer::Serialize(FArchive& Ar, UPackageMap* Map, bool& bOutSuccess)
+{
+	bOutSuccess = true;
+
+	int Number = InputValues.Num();
+	Ar << Number;
+	if (Ar.IsLoading())
+	{
+		InputValues.Reserve(Number);
+		for (int I = 0; I < Number; I++)
+		{
+			FModuleInputValue Value;
+			Value.Serialize(Ar, Map, bOutSuccess);
+			InputValues.EmplaceAt(I, Value);
+		}
+	}
+	else
+	{
+		for (int I = 0; I < Number; I++)
+		{
+			InputValues[I].Serialize(Ar, Map, bOutSuccess);
+		}
+	}
+
+	return;
+}
+
+int FModuleInputContainer::AddInput(EModuleInputValueType Type, TSubclassOf<UDefaultModularVehicleInputModifier>& InputModifierClass)
+{
+	InputValues.Add(FModuleInputValue(Type, FVector::ZeroVector));
+	return InputValues.Num() - 1;
+}
+
+void FModuleInputContainer::RemoveAllInputs()
+{
+	InputValues.Reset();
+}
+
+void FModuleInputContainer::Lerp(const FModuleInputContainer& Min, const FModuleInputContainer& Max, float Alpha)
+{
+	for (int I = 0; I < InputValues.Num(); I++)
+	{
+		InputValues[I].Lerp(Min.InputValues[I], Max.InputValues[I], Alpha);
+	}
+}
+
+void FModuleInputContainer::Merge(const FModuleInputContainer& From)
+{
+	for (int I = 0; I < InputValues.Num(); I++)
+	{
+		InputValues[I].Merge(From.InputValues[I]);
+	}
+}
+
+
+void FInputInterface::SetValue(const FName& InName, const FModuleInputValue& InValue)
+{
+	if (const int* Index = NameMap.Find(InName))
+	{
+		ValueContainer.SetValueAtIndex(*Index, InValue);
+	}
+	else
+	{
+		UE_LOG(LogModularInput, Warning, TEXT("Trying to set the value of an undefined control input %s"), *InName.ToString());
+	}
+}
+
+FModuleInputValue FInputInterface::GetValue(const FName& InName) const
+{
+	if (const int* Index = NameMap.Find(InName))
+	{
+		return ValueContainer.GetValueAtIndex(*Index);
+	}
+	else
+	{
+		UE_LOG(LogModularInput, Warning, TEXT("Trying to get the value of an undefined control input %s"), *InName.ToString());
+	}
+
+	return FModuleInputValue(EModuleInputValueType::MBoolean, FVector::ZeroVector);
+}
+
+float FInputInterface::GetMagnitude(const FName& InName) const
+{
+	if (const int* Index = NameMap.Find(InName))
+	{
+		return ValueContainer.GetValueAtIndex(*Index).GetMagnitude();
+	}
+
+	return 0.0f;
+}
+
+bool FInputInterface::InputsNonZero() const
+{
+	for (int I = 0; I < ValueContainer.GetNumInputs(); I++)
+	{
+		if (ValueContainer.GetValueAtIndex(I).IsNonZero())
+		{
+			return true;
+		}
+	}
+
+	return false;
+}
\ No newline at end of file
diff --git a/Engine/Source/Runtime/Experimental/ChaosVehicles/ChaosVehiclesCore/Private/SimModule/MotorModule.cpp b/Engine/Source/Runtime/Experimental/ChaosVehicles/ChaosVehiclesCore/Private/SimModule/MotorModule.cpp
--- a/Engine/Source/Runtime/Experimental/ChaosVehicles/ChaosVehiclesCore/Private/SimModule/MotorModule.cpp
+++ b/Engine/Source/Runtime/Experimental/ChaosVehicles/ChaosVehiclesCore/Private/SimModule/MotorModule.cpp
@@ -11,7 +11,6 @@ UE_DISABLE_OPTIMIZATION
 
 namespace Chaos
 {
-
 	FMotorSimModule::FMotorSimModule(const FMotorSettings& Settings)
 		: TSimModuleSettings<FMotorSettings>(Settings)
 	{
@@ -36,7 +35,7 @@ namespace Chaos
 				NormalizedTorque = -1.0f - FMath::Square(1.0f + 2.0f * NormalizedRPM);
 			}
 
-			DriveTorque = NormalizedTorque * Setup().MaxTorque * Inputs.ControlInputs.Throttle;
+			DriveTorque = NormalizedTorque * Setup().MaxTorque * Inputs.GetControls().GetMagnitude(ThrottleControlName);
 
 			float BrakeTorque = 0.f;
 			TransmitTorque(VehicleModuleSystem, DriveTorque, BrakeTorque);
diff --git a/Engine/Source/Runtime/Experimental/ChaosVehicles/ChaosVehiclesCore/Private/SimModule/SimModuleTree.cpp b/Engine/Source/Runtime/Experimental/ChaosVehicles/ChaosVehiclesCore/Private/SimModule/SimModuleTree.cpp
--- a/Engine/Source/Runtime/Experimental/ChaosVehicles/ChaosVehiclesCore/Private/SimModule/SimModuleTree.cpp
+++ b/Engine/Source/Runtime/Experimental/ChaosVehicles/ChaosVehiclesCore/Private/SimModule/SimModuleTree.cpp
@@ -230,7 +230,7 @@ void FSimModuleTree::Simulate(float DeltaTime, FAllInputs& Inputs, FClusterUnion
 		{
 			UpdateVehicleState(PhysicsProxy);
 
-			UpdateModuleVelocites(PhysicsProxy, Inputs.ControlInputs.InputNonZero() || Inputs.bKeepVehicleAwake);
+			UpdateModuleVelocites(PhysicsProxy, Inputs.GetControls().InputsNonZero() || Inputs.bKeepVehicleAwake);
 		}
 
 		TArray<int> RootNodes;
diff --git a/Engine/Source/Runtime/Experimental/ChaosVehicles/ChaosVehiclesCore/Private/SimModule/ThrusterModule.cpp b/Engine/Source/Runtime/Experimental/ChaosVehicles/ChaosVehiclesCore/Private/SimModule/ThrusterModule.cpp
--- a/Engine/Source/Runtime/Experimental/ChaosVehicles/ChaosVehiclesCore/Private/SimModule/ThrusterModule.cpp
+++ b/Engine/Source/Runtime/Experimental/ChaosVehicles/ChaosVehiclesCore/Private/SimModule/ThrusterModule.cpp
@@ -22,12 +22,12 @@ namespace Chaos
 		SteerAngleDegrees = 0.0f;
 		if (Setup().SteeringEnabled)
 		{
-			SteerAngleDegrees = Setup().SteeringEnabled ? Inputs.ControlInputs.Steering * Setup().MaxSteeringAngle : 0.0f;
+			SteerAngleDegrees = Setup().SteeringEnabled ? Inputs.GetControls().GetMagnitude(TEXT("Steering")) * Setup().MaxSteeringAngle : 0.0f;
 		}
 
 		// applies continuous force
-		float BoostEffect = Inputs.ControlInputs.Boost * Setup().BoostMultiplier;
-		FVector Force = Setup().ForceAxis * Setup().MaxThrustForce * Inputs.ControlInputs.Throttle * (1.0f + BoostEffect);
+		float BoostEffect = Inputs.GetControls().GetMagnitude(BoostControlName) * Setup().BoostMultiplier;
+		FVector Force = Setup().ForceAxis * Setup().MaxThrustForce * Inputs.GetControls().GetMagnitude(ThrottleControlName) * (1.0f + BoostEffect);
 		FQuat Steer = FQuat(Setup().SteeringAxis, FMath::DegreesToRadians(SteerAngleDegrees) * Setup().SteeringForceEffect);
 		AddLocalForceAtPosition(Steer.RotateVector(Force), Setup().ForceOffset, true, false, false, FColor::Magenta);
 	}
diff --git a/Engine/Source/Runtime/Experimental/ChaosVehicles/ChaosVehiclesCore/Private/SimModule/TransmissionModule.cpp b/Engine/Source/Runtime/Experimental/ChaosVehicles/ChaosVehiclesCore/Private/SimModule/TransmissionModule.cpp
--- a/Engine/Source/Runtime/Experimental/ChaosVehicles/ChaosVehiclesCore/Private/SimModule/TransmissionModule.cpp
+++ b/Engine/Source/Runtime/Experimental/ChaosVehicles/ChaosVehiclesCore/Private/SimModule/TransmissionModule.cpp
@@ -30,7 +30,7 @@ namespace Chaos
 
 		if (Setup().AutoReverse)
 		{
-			if (Inputs.ControlInputs.IsReversing)
+			if (Inputs.bIsReversing)
 			{
 				// if reversing change to reverse gear if currently in a forwards gear
 				if (TargetGear > 0)
@@ -87,18 +87,18 @@ namespace Chaos
 		}
 		else
 		{
-			if (Inputs.ControlInputs.ChangeUp)
+			if (Inputs.GetControls().GetMagnitude(ChangeUpControlName))
 			{
 				ChangeUp();
 			}
-			else if (Inputs.ControlInputs.ChangeDown)
+			else if (Inputs.GetControls().GetMagnitude(ChangeDownControlName))
 			{
 				ChangeDown();
 			}
-			else if (Inputs.ControlInputs.GearNumber)
-			{
-				TargetGear = Inputs.ControlInputs.GearNumber;
-			}
+			//else if (Inputs.ControlInputs.GearNumber) #TODO: this was never used
+			//{
+			//	TargetGear = Inputs.ControlInputs.GearNumber;
+			//}
 		}
 
 		if (CurrentGear != TargetGear)
diff --git a/Engine/Source/Runtime/Experimental/ChaosVehicles/ChaosVehiclesCore/Private/SimModule/WheelModule.cpp b/Engine/Source/Runtime/Experimental/ChaosVehicles/ChaosVehiclesCore/Private/SimModule/WheelModule.cpp
--- a/Engine/Source/Runtime/Experimental/ChaosVehicles/ChaosVehiclesCore/Private/SimModule/WheelModule.cpp
+++ b/Engine/Source/Runtime/Experimental/ChaosVehicles/ChaosVehiclesCore/Private/SimModule/WheelModule.cpp
@@ -10,7 +10,6 @@ UE_DISABLE_OPTIMIZATION_SHIP
 
 namespace Chaos
 {
-
 	FWheelSimModule::FWheelSimModule(const FWheelSettings& Settings)
 		: TSimModuleSettings<FWheelSettings>(Settings)
 		, BrakeTorque(0.0f)
@@ -32,9 +31,9 @@ namespace Chaos
 		float TorqueScaling = 1.0f;
 		float TractionControlAndABSScaling = 0.98f;	// how close to perfection is the system working
 
-		float HandbrakeTorque = Setup().HandbrakeEnabled ? Inputs.ControlInputs.Handbrake * Setup().HandbrakeTorque : 0.0f;
-		SteerAngleDegrees = Setup().SteeringEnabled ? Inputs.ControlInputs.Steering * Setup().MaxSteeringAngle : 0.0f;
-		BrakeTorque = Inputs.ControlInputs.Brake * Setup().MaxBrakeTorque + HandbrakeTorque;
+		float HandbrakeTorque = Setup().HandbrakeEnabled ? Inputs.GetControls().GetMagnitude(HandbrakeControlName) * Setup().HandbrakeTorque : 0.0f;
+		SteerAngleDegrees = Setup().SteeringEnabled ? Inputs.GetControls().GetMagnitude(SteeringControlName) * Setup().MaxSteeringAngle : 0.0f;
+		BrakeTorque = Inputs.GetControls().GetMagnitude(BrakeControlName) * Setup().MaxBrakeTorque + HandbrakeTorque;
 		LoadTorque = 0.0f;
 		ForceFromFriction = FVector::ZeroVector;
 		float TorqueFromGroundInteraction = 0.0f;
@@ -42,7 +41,7 @@ namespace Chaos
 
 		// TODO: think about doing this properly, stops vehicles rolling around on their own too much
 		// i.e. an auto handbrake feature
-		if (Setup().AutoHandbrakeEnabled && LocalLinearVelocity.X < Setup().AutoHandbrakeVelocityThreshold && (Inputs.ControlInputs.Brake < SMALL_NUMBER && Inputs.ControlInputs.Throttle < SMALL_NUMBER))
+		if (Setup().AutoHandbrakeEnabled && LocalLinearVelocity.X < Setup().AutoHandbrakeVelocityThreshold && (Inputs.GetControls().GetMagnitude(BrakeControlName) < SMALL_NUMBER && Inputs.GetControls().GetMagnitude(ThrottleControlName) < SMALL_NUMBER))
 		{
 			BrakeTorque = Setup().HandbrakeTorque;
 		}
diff --git a/Engine/Source/Runtime/Experimental/ChaosVehicles/ChaosVehiclesCore/Public/SimModule/ModuleInput.h b/Engine/Source/Runtime/Experimental/ChaosVehicles/ChaosVehiclesCore/Public/SimModule/ModuleInput.h
new file mode 100644
--- /dev/null
+++ b/Engine/Source/Runtime/Experimental/ChaosVehicles/ChaosVehiclesCore/Public/SimModule/ModuleInput.h
@@ -0,0 +1,420 @@
+// Copyright Epic Games, Inc. All Rights Reserved.
+
+#pragma once
+
+#include "CoreMinimal.h"
+#include "Templates/SubclassOf.h"
+#include "ModuleInput.generated.h"
+
+DECLARE_LOG_CATEGORY_EXTERN(LogModularInput, Log, All);
+
+
+UENUM(BlueprintType)
+enum class EModuleInputValueType : uint8
+{
+	// Value types in increasing size order (used for type promotion)
+
+	MBoolean	UMETA(DisplayName = "Digital (bool)"),
+	MAxis1D		UMETA(DisplayName = "Axis1D (float)"),
+	MAxis2D		UMETA(DisplayName = "Axis2D (Vector2D)"),
+	MAxis3D		UMETA(DisplayName = "Axis3D (Vector)"),
+};
+
+/** Input Options */
+UENUM(BlueprintType)
+enum class EFunctionType : uint8
+{
+	LinearFunction = 0,
+	SquaredFunction,
+	CustomCurve
+};
+
+
+USTRUCT(BlueprintType)
+struct CHAOSVEHICLESCORE_API FModuleInputValue
+{
+	GENERATED_BODY()
+
+public:
+	using MAxis1D = double; // #TODO: Axis1D (float) double or float? FVector2D and FVector are double
+	using MAxis2D = FVector2D;
+	using MAxis3D = FVector;
+
+	// Support all relevant default constructors (FModuleInputValue isn't movable)
+	FModuleInputValue() = default;
+	FModuleInputValue(const FModuleInputValue&) = default;
+	FModuleInputValue& operator= (const FModuleInputValue&) = default;
+
+	// Specialized constructors for supported types
+	// Converting a value to a different type (e.g. Val = FVector(1, 1, 1); Val = true;) zeroes out any unused components to ensure getters continue to function correctly.
+	explicit FModuleInputValue(bool bInValue) : Value(bInValue ? 1.f : 0.f, 0.f, 0.f), ValueType(EModuleInputValueType::MBoolean) {}
+	FModuleInputValue(MAxis1D InValue) : Value(InValue, 0.f, 0.f), ValueType(EModuleInputValueType::MAxis1D) {}
+	FModuleInputValue(MAxis2D InValue) : Value(InValue, 0.f), ValueType(EModuleInputValueType::MAxis2D) {}
+	FModuleInputValue(MAxis3D InValue) : Value(InValue), ValueType(EModuleInputValueType::MAxis3D) {}
+
+	// Build a specific type with an arbitrary Axis3D value
+	FModuleInputValue(EModuleInputValueType InValueType, MAxis3D InValue) : Value(InValue), ValueType(InValueType)
+	{
+		// Clear out value components to match type
+		switch (ValueType)
+		{
+		case EModuleInputValueType::MBoolean:
+		case EModuleInputValueType::MAxis1D:
+			Value.Y = 0.f;
+			//[[fallthrough]];
+		case EModuleInputValueType::MAxis2D:
+			Value.Z = 0.f;
+			//[[fallthrough]];
+		case EModuleInputValueType::MAxis3D:
+		default:
+			return;
+		}
+	}
+
+	// Resets Value without affecting ValueType
+	void Reset()
+	{
+		Value = FVector::ZeroVector;
+	}
+
+	FModuleInputValue& operator+=(const FModuleInputValue& Rhs)
+	{
+		Value += Rhs.Value;
+		// Promote value type to largest number of bits.
+		ValueType = FMath::Max(ValueType, Rhs.ValueType);
+		return *this;
+	}
+
+	friend FModuleInputValue operator+(const FModuleInputValue& Lhs, const FModuleInputValue& Rhs)
+	{
+		FModuleInputValue Result(Lhs);
+		Result += Rhs;
+		return Result;
+	}
+
+	FModuleInputValue& operator-=(const FModuleInputValue& Rhs)
+	{
+		Value -= Rhs.Value;
+		// Promote value type to largest number of bits.
+		ValueType = FMath::Max(ValueType, Rhs.ValueType);
+		return *this;
+	}
+
+	friend FModuleInputValue operator-(const FModuleInputValue& Lhs, const FModuleInputValue& Rhs)
+	{
+		FModuleInputValue Result(Lhs);
+		Result -= Rhs;
+		return Result;
+	}
+
+	// Scalar operators
+	FModuleInputValue& operator*=(float Scalar)
+	{
+		Value *= Scalar;
+		return *this;
+	}
+
+	friend FModuleInputValue operator*(const FModuleInputValue& Lhs, const float Rhs)
+	{
+		FModuleInputValue Result(Lhs);
+		Result *= Rhs;
+		return Result;
+	}
+
+	static FModuleInputValue Clamp(const FModuleInputValue& InValue, const float InMin, const float InMax)
+	{
+		FModuleInputValue OutValue = InValue;
+		float Mag = InValue.GetMagnitude();
+
+		if (Mag < InMin)
+		{
+			OutValue.SetMagnitude(InMin);
+		}
+		else if (Mag > InMax)
+		{
+			OutValue.SetMagnitude(InMax);
+		}
+
+		return OutValue;
+	}
+
+
+
+	template<typename T>
+	inline T Get() const { static_assert(sizeof(T) == 0, "Unsupported conversion for type"); }
+
+	// Read only index based value accessor, doesn't care about type. Expect 0 when accessing unused components.
+	float operator[](int32 Index) const { return Value[Index]; }
+
+	bool IsNonZero(float Tolerance = KINDA_SMALL_NUMBER) const { return Value.SizeSquared() >= Tolerance * Tolerance; }
+
+	// In-place type conversion
+	FModuleInputValue& ConvertToType(EModuleInputValueType Type)
+	{
+		if (ValueType != Type)
+		{
+			*this = FModuleInputValue(Type, Value);
+		}
+		return *this;
+	}
+	FModuleInputValue& ConvertToType(const FModuleInputValue& Other) { return ConvertToType(Other.GetValueType()); }
+
+	EModuleInputValueType GetValueType() const { return ValueType; }
+
+	float GetMagnitudeSq() const;
+	float GetMagnitude() const;
+
+	// Serialize values, not quantized yet and not serializing deltas
+	void Serialize(FArchive& Ar, UPackageMap* Map, bool& bOutSuccess);
+
+	void Lerp(const FModuleInputValue& Min, const FModuleInputValue& Max, float Alpha)
+	{
+		Value = FMath::Lerp(Min.Value, Max.Value, Alpha);
+	}
+
+	void Merge(const FModuleInputValue& From);
+
+	// Type sensitive debug stringify
+	FString ToString() const;
+
+protected:
+	void SetMagnitude(float NewSize);
+
+	UPROPERTY()
+	FVector Value = FVector::ZeroVector;
+
+	UPROPERTY()
+	EModuleInputValueType ValueType = EModuleInputValueType::MBoolean;
+};
+
+// Supported getter specializations
+template<>
+inline bool FModuleInputValue::Get() const
+{
+	// True if any component is non-zero
+	return IsNonZero();
+}
+
+template<>
+inline FModuleInputValue::MAxis1D FModuleInputValue::Get() const
+{
+	return Value.X;
+}
+
+template<>
+inline FModuleInputValue::MAxis2D FModuleInputValue::Get() const
+{
+	return MAxis2D(Value.X, Value.Y);
+}
+
+template<>
+inline FModuleInputValue::MAxis3D FModuleInputValue::Get() const
+{
+	return Value;
+}
+
+
+class CHAOSVEHICLESCORE_API FModuleInputConversion
+{
+public:
+
+	static bool ToBool(FModuleInputValue& InValue)
+	{
+		return InValue.Get<bool>();
+	}
+
+	static float ToAxis1D(FModuleInputValue& InValue)
+	{
+		return static_cast<float>(InValue.Get<FModuleInputValue::MAxis1D>());
+	}
+
+	static FVector2D ToAxis2D(FModuleInputValue& InValue)
+	{
+		return InValue.Get<FModuleInputValue::MAxis2D>();
+	}
+
+	static FVector ToAxis3D(FModuleInputValue& InValue)
+	{
+		return InValue.Get<FModuleInputValue::MAxis3D>();
+	}
+
+	static FString ToString(FModuleInputValue& ActionValue)
+	{
+		return ActionValue.ToString();
+	}
+};
+
+
+UCLASS(BlueprintType, Blueprintable)
+class CHAOSVEHICLESCORE_API UDefaultModularVehicleInputModifier : public UObject
+{
+	GENERATED_BODY()
+
+public:
+	UDefaultModularVehicleInputModifier(const class FObjectInitializer& ObjectInitializer)
+		: Super(ObjectInitializer)
+		, RiseRate(5.0f), FallRate(5.0f), InputCurveFunction(EFunctionType::LinearFunction) { }
+
+	virtual ~UDefaultModularVehicleInputModifier()
+	{
+	}
+
+	/**
+		* Rate at which the input value rises
+		*/
+	UPROPERTY(EditAnywhere, Category = VehicleInputRate)
+	float RiseRate;
+
+	/**
+	 * Rate at which the input value falls
+	 */
+	UPROPERTY(EditAnywhere, Category = VehicleInputRate)
+	float FallRate;
+
+	/**
+	 * Controller input curve, various predefined options, linear, squared, or user can specify a custom curve function
+	 */
+	UPROPERTY(EditAnywhere, Category = VehicleInputRate)
+	EFunctionType InputCurveFunction;
+
+	/**
+	 * Controller input curve - should be a normalized float curve, i.e. time from 0 to 1 and values between 0 and 1
+	 * This curve is only sued if the InputCurveFunction above is set to CustomCurve
+	 */
+	//UPROPERTY(EditAnywhere, Category = VehicleInputRate) #TODO: Reinstate?
+	//FRuntimeFloatCurve UserCurve;
+
+	/** Change an output value using max rise and fall rates */
+	virtual FModuleInputValue InterpInputValue(float DeltaTime, const FModuleInputValue& CurrentValue, const FModuleInputValue& NewValue) const;
+	virtual float CalcControlFunction(float InputValue);
+};
+
+
+USTRUCT(BlueprintType)
+struct CHAOSVEHICLESCORE_API FModuleInputSetup
+{
+	GENERATED_BODY()
+
+	FModuleInputSetup()
+	{
+		Type = EModuleInputValueType::MBoolean;
+	}
+
+	FModuleInputSetup(const FName& InName, const EModuleInputValueType& InType)
+		: Name(InName)
+		, Type(InType)
+	{
+	}
+
+	bool operator==(const FModuleInputSetup& Rhs) const
+	{
+		return (Name == Rhs.Name);
+	}
+
+	UPROPERTY(EditAnywhere, Category = VehicleInput)
+	FName Name;
+
+	UPROPERTY(EditAnywhere, Category = VehicleInput)
+	EModuleInputValueType Type;
+
+	UPROPERTY(EditAnywhere, Category = VehicleInput)
+	TSubclassOf<UDefaultModularVehicleInputModifier> InputModifierClass;
+};
+
+class CHAOSVEHICLESCORE_API FScopedModuleInputInitializer
+{
+public:
+	FScopedModuleInputInitializer(TArray<struct FModuleInputSetup>& InSetupData)
+	{
+		InitSetupData = &InSetupData;
+	}
+
+	~FScopedModuleInputInitializer()
+	{
+		InitSetupData = nullptr;
+	}
+
+	static bool HasSetup() { return InitSetupData != nullptr; }
+	static TArray<struct FModuleInputSetup>* GetSetup() { return InitSetupData; }
+	
+private:
+
+	static TArray<struct FModuleInputSetup>* InitSetupData;
+};
+
+
+USTRUCT()
+struct CHAOSVEHICLESCORE_API FModuleInputContainer
+{
+	GENERATED_BODY()
+
+public:
+	using FInputNameMap = TMap<FName, int>;
+	using FInputValues = TArray<FModuleInputValue>;
+
+	FModuleInputContainer()
+	{
+		if (FScopedModuleInputInitializer::HasSetup())
+		{
+			FInputNameMap NameMapOut;
+			Initialize(*FScopedModuleInputInitializer::GetSetup(), NameMapOut);
+		}
+	}
+
+	int GetNumInputs() const { return InputValues.Num(); }
+	FModuleInputValue GetValueAtIndex(int Index) const { return InputValues[Index]; }
+	void SetValueAtIndex(int Index, const FModuleInputValue& InValue) { InputValues[Index] = InValue; }
+
+	FModuleInputContainer& operator=(const FModuleInputContainer& Other)
+	{
+		if (&Other != this)
+		{
+			// perform a deep copy
+			if (Other.InputValues.Num())
+			{
+				InputValues = Other.InputValues;
+			}
+		}
+		return *this;
+	}
+
+	void Initialize(TArray<FModuleInputSetup>& SetupData, FInputNameMap& NameMapOut);
+
+	void Serialize(FArchive& Ar, UPackageMap* Map, bool& bOutSuccess);
+
+	int AddInput(EModuleInputValueType Type, TSubclassOf<UDefaultModularVehicleInputModifier>& InputModifierClass);
+
+	void RemoveAllInputs();
+
+	void Lerp(const FModuleInputContainer& Min, const FModuleInputContainer& Max, float Alpha);
+
+	void Merge(const FModuleInputContainer& From);
+
+private:
+	UPROPERTY()
+	TArray<FModuleInputValue> InputValues;
+
+};
+
+
+class CHAOSVEHICLESCORE_API FInputInterface
+{
+public:
+
+	using FInputNameMap = TMap<FName, int>;
+
+	FInputInterface(const FInputNameMap& InNameMap, FModuleInputContainer& InValueContainer)
+		: NameMap(InNameMap)
+		, ValueContainer(InValueContainer)
+	{
+	}
+
+	void SetValue(const FName& InName, const FModuleInputValue& InValue);
+	FModuleInputValue GetValue(const FName& InName) const;
+	float GetMagnitude(const FName& InName) const;
+	bool InputsNonZero() const;
+
+	const FInputNameMap& NameMap;			// per vehicle
+	FModuleInputContainer& ValueContainer;	// per vehicle instance
+};
+
diff --git a/Engine/Source/Runtime/Experimental/ChaosVehicles/ChaosVehiclesCore/Public/SimModule/SimModuleTree.h b/Engine/Source/Runtime/Experimental/ChaosVehicles/ChaosVehiclesCore/Public/SimModule/SimModuleTree.h
--- a/Engine/Source/Runtime/Experimental/ChaosVehicles/ChaosVehiclesCore/Public/SimModule/SimModuleTree.h
+++ b/Engine/Source/Runtime/Experimental/ChaosVehicles/ChaosVehiclesCore/Public/SimModule/SimModuleTree.h
@@ -4,6 +4,7 @@
 
 #include "SimModule/DeferredForcesModular.h"
 #include "SimModule/SimulationModuleBase.h"
+#include "SimModule/ModuleInput.h"
 
 
 DECLARE_STATS_GROUP(TEXT("ModularVehicle.SimTree"), STATGROUP_ModularVehicleSimTree, STATGROUP_Advanced);
@@ -158,10 +159,6 @@ namespace Chaos
 		void SetSimulationEnabled(bool bInEnabled) { bSimulationEnabled = bInEnabled; }
 		bool IsSimulationEnabled() { return bSimulationEnabled; }
 
-		FControlInputs& GetControlInputs()
-		{
-			return AllInputs.ControlInputs;
-		}
 
 		const FVehicleState& GetVehicleState() const
 		{
diff --git a/Engine/Source/Runtime/Experimental/ChaosVehicles/ChaosVehiclesCore/Public/SimModule/SimulationModuleBase.h b/Engine/Source/Runtime/Experimental/ChaosVehicles/ChaosVehiclesCore/Public/SimModule/SimulationModuleBase.h
--- a/Engine/Source/Runtime/Experimental/ChaosVehicles/ChaosVehiclesCore/Public/SimModule/SimulationModuleBase.h
+++ b/Engine/Source/Runtime/Experimental/ChaosVehicles/ChaosVehiclesCore/Public/SimModule/SimulationModuleBase.h
@@ -7,6 +7,7 @@
 #include "Chaos/ParticleHandleFwd.h"
 #include "Chaos/GeometryParticlesfwd.h"
 #include "SimModule/ModuleFactoryRegister.h"
+#include "SimModule/ModuleInput.h"
 
 DECLARE_LOG_CATEGORY_EXTERN(LogSimulationModule, Warning, All);
 
@@ -26,58 +27,17 @@ namespace Chaos
 	struct FSimOutputData;
 	class FClusterUnionPhysicsProxy;
 
-	struct CHAOSVEHICLESCORE_API FControlInputs
-	{
-
-		FControlInputs()
-			: IsValid(false)
-			, IsReversing(false)
-			, Throttle(0)
-			, Brake(0)
-			, Steering(0)
-			, Clutch(0)
-			, Handbrake(0)
-			, Roll(0)
-			, Pitch(0)
-			, Yaw(0)
-			, Boost(0)
-			, Drift(0)
-			, ChangeUp(false)
-			, ChangeDown(false)
-			, GearNumber(0)
-			, InputDebugIndex(0)
-		{
-		}
-
-		bool InputNonZero()
-		{
-			return FMath::Abs(Throttle) > SMALL_NUMBER
-				|| FMath::Abs(Brake) > SMALL_NUMBER
-				|| FMath::Abs(Steering) > SMALL_NUMBER
-				|| FMath::Abs(Roll) > SMALL_NUMBER
-				|| FMath::Abs(Pitch) > SMALL_NUMBER
-				|| FMath::Abs(Yaw) > SMALL_NUMBER
-				|| FMath::Abs(Boost) > SMALL_NUMBER
-				|| FMath::Abs(Drift) > SMALL_NUMBER;
-		}
-
-		bool IsValid;
-		bool IsReversing;
-		float Throttle;
-		float Brake;
-		float Steering;
-		float Clutch;
-		float Handbrake;
-		float Roll;
-		float Pitch;
-		float Yaw;
-		float Boost;
-		float Drift;
-		bool ChangeUp;
-		bool ChangeDown;
-		int GearNumber;
-		int InputDebugIndex;
-	};
+	const FName HandbrakeControlName("Handbrake");
+	const FName ThrottleControlName("Throttle");
+	const FName SteeringControlName("Steering");
+	const FName BrakeControlName("Brake");
+	const FName ClutchControlName("Clutch");
+	const FName BoostControlName("Boost");
+	const FName ChangeUpControlName("ChangeUp");
+	const FName ChangeDownControlName("ChangeDown");
+	const FName PitchControlName("Pitch");
+	const FName RollControlName("Roll");
+	const FName YawControlName("Yaw");
 
 	struct CHAOSVEHICLESCORE_API FModuleHitResults
 	{
@@ -89,10 +49,13 @@ namespace Chaos
 
 	struct CHAOSVEHICLESCORE_API FAllInputs
 	{
+		FInputInterface& GetControls() const { check(ControlInputs); return *ControlInputs; }
+
 		FTransform VehicleWorldTransform;
 		TMap<int32, FModuleHitResults> HitResults;
-		FControlInputs ControlInputs;
-		bool bKeepVehicleAwake;
+		FInputInterface* ControlInputs = nullptr;
+		bool bIsReversing = false;
+		bool bKeepVehicleAwake = false;
 	};
 
 	/**
