Move the physics constraint into the suspension module that requires it. Remove references to core modules only referencing through base class.
Also renaming network struc names from Datas to Data

diff --git a/Engine/Plugins/Experimental/ChaosModularVehicle/Source/ChaosModularVehicleEngine/Private/ChaosModularVehicle/ModularVehicleBaseComponent.cpp b/Engine/Plugins/Experimental/ChaosModularVehicle/Source/ChaosModularVehicleEngine/Private/ChaosModularVehicle/ModularVehicleBaseComponent.cpp
--- a/Engine/Plugins/Experimental/ChaosModularVehicle/Source/ChaosModularVehicleEngine/Private/ChaosModularVehicle/ModularVehicleBaseComponent.cpp
+++ b/Engine/Plugins/Experimental/ChaosModularVehicle/Source/ChaosModularVehicleEngine/Private/ChaosModularVehicle/ModularVehicleBaseComponent.cpp
@@ -17,7 +17,6 @@
 #include "GameFramework/Pawn.h"
 #include "Physics/Experimental/PhysScene_Chaos.h"
 #include "SimModule/SimModuleTree.h"
-#include "Chaos/PBDSuspensionConstraints.h"
 
 #include "ChaosModularVehicle/ChaosSimModuleManagerAsyncCallback.h"
 
@@ -306,12 +305,9 @@ void UModularVehicleBaseComponent::CreateAssociatedSimComponents(UPrimitiveCompo
 		}
 		ParentIndex = TreeIndex;
 
-		if (bModularVehicle_SuspensionConstraint_Enabled)
+		if (Chaos::FClusterUnionPhysicsProxy* Proxy = ClusterUnionComponent->GetPhysicsProxy())
 		{
-			if (NewModule->GetSimType() == Chaos::eSimType::Suspension)
-			{
-				CreateConstraint(NewModule);
-			}
+			NewModule->OnConstruction_External(Proxy);
 		}
 	}
 
@@ -396,99 +392,6 @@ void UModularVehicleBaseComponent::TickComponent(float DeltaTime, ELevelTick Tic
 
 }
 
-void UModularVehicleBaseComponent::CreateConstraint(Chaos::ISimulationModuleBase* NewModule)
-{
-	check(NewModule->GetSimType() == Chaos::eSimType::Suspension);
-	Chaos::FSuspensionSimModule* SuspensionModule = static_cast<Chaos::FSuspensionSimModule*>(NewModule);
-
-	const Chaos::FSuspensionSettings& Setup = SuspensionModule->Setup();
-	const FVector& LocalOffset = SuspensionModule->GetInitialParticleTransform().GetLocation();
-
-	if (ClusterUnionComponent && ClusterUnionComponent->GetPhysicsProxy())
-	{
-		const Chaos::FPhysicsObjectHandle& PhysicsObject = ClusterUnionComponent->GetPhysicsProxy()->GetPhysicsObjectHandle();
-
-		if ( FChaosScene* Scene = static_cast<FChaosScene*>(FPhysicsObjectExternalInterface::GetScene({ &PhysicsObject, 1 })) )
-		{
-			FLockedWritePhysicsObjectExternalInterface Interface = FPhysicsObjectExternalInterface::LockWrite(Scene);
-			if (const Chaos::FGeometryParticle* Particle = Interface->GetParticle(PhysicsObject))
-			{
-				FPhysicsConstraintHandle ConstraintHandle = FPhysicsInterface::CreateSuspension(PhysicsObject, LocalOffset);
-
-				if (ConstraintHandle.IsValid())
-				{
-					ConstraintHandles.Add(ConstraintHandle);
-
-					if (Chaos::FSuspensionConstraint* Constraint = static_cast<Chaos::FSuspensionConstraint*>(ConstraintHandle.Constraint))
-					{
-						Constraint->SetHardstopStiffness(1.0f);
-						Constraint->SetSpringStiffness(Setup.SpringRate * 0.25f);
-						Constraint->SetSpringPreload(Setup.SpringPreload);
-						Constraint->SetSpringDamping(Setup.SpringDamping * 5.0f);
-						Constraint->SetMinLength(-Setup.MaxRaise);
-						Constraint->SetMaxLength(Setup.MaxDrop);
-						Constraint->SetAxis(-Setup.SuspensionAxis);
-
-						SuspensionModule->SetSuspensionConstraint(Constraint);
-						SuspensionModule->SetConstraintIndex(ConstraintHandles.Num()-1);
-					}
-				}
-			}
-		}
-	}
-}
-
-void UModularVehicleBaseComponent::DestroyConstraint(int ConstraintIndex)
-{
-	if (ConstraintIndex >= 0 && ConstraintIndex < ConstraintHandles.Num())
-	{
-		FPhysicsConstraintHandle ConstraintHandle = ConstraintHandles[ConstraintIndex];
-		FPhysicsCommand::ExecuteWrite(ConstraintHandle, [&](const FPhysicsConstraintHandle& Constraint)
-			{
-				FPhysicsInterface::ReleaseConstraint(ConstraintHandle);
-			});
-
-		ConstraintHandles[ConstraintIndex].Reset();
-	}
-
-}
-
-void UModularVehicleBaseComponent::DestroyAllConstraints()
-{
-	if (ConstraintHandles.Num() > 0)
-	{
-		for (FPhysicsConstraintHandle ConstraintHandle : ConstraintHandles)
-		{
-			if (ConstraintHandle.IsValid())
-			{
-				FPhysicsCommand::ExecuteWrite(ConstraintHandle, [&](const FPhysicsConstraintHandle& Constraint)
-					{
-						FPhysicsInterface::ReleaseConstraint(ConstraintHandle);
-					});
-			}
-		}
-	}
-	ConstraintHandles.Empty();
-
-}
-
-void UModularVehicleBaseComponent::EnableConstraint(int ConstraintIndex, bool bEnabled)
-{
-	if (ConstraintIndex >= 0 && ConstraintIndex < ConstraintHandles.Num())
-	{
-		FPhysicsConstraintHandle ConstraintHandle = ConstraintHandles[ConstraintIndex];
-		FPhysicsCommand::ExecuteWrite(ConstraintHandle, [&](const FPhysicsConstraintHandle& Constraint)
-			{
-				if (Chaos::FSuspensionConstraint* SuspensionConstraint = static_cast<Chaos::FSuspensionConstraint*>(ConstraintHandle.Constraint))
-				{
-					SuspensionConstraint->SetEnabled(bEnabled);
-				}
-			});
-
-		ConstraintHandles[ConstraintIndex].Reset();
-	}
-
-}
 
 int32 UModularVehicleBaseComponent::FindComponentAddOrder(UPrimitiveComponent* InComponent)
 {
@@ -963,13 +866,7 @@ void UModularVehicleBaseComponent::RemoveComponentFromSimulation(UPrimitiveCompo
 					{
 						SimModule->SetAnimationEnabled(false);
 						SimModule->SetStateFlags(Chaos::eSimModuleState::Disabled);
-
-						if (SimModule->GetSimType() == Chaos::eSimType::Suspension)
-						{
-							Chaos::FSuspensionSimModule* SuspensionModule = static_cast<Chaos::FSuspensionSimModule*>(SimModule);
-							DestroyConstraint(SuspensionModule->GetConstraintIndex());
-						}
-
+						SimModule->OnTermination_External();
 						break;
 					}
 				}
diff --git a/Engine/Plugins/Experimental/ChaosModularVehicle/Source/ChaosModularVehicleEngine/Private/ChaosModularVehicle/SuspensionSimModule.cpp b/Engine/Plugins/Experimental/ChaosModularVehicle/Source/ChaosModularVehicleEngine/Private/ChaosModularVehicle/SuspensionSimModule.cpp
new file mode 100644
--- /dev/null
+++ b/Engine/Plugins/Experimental/ChaosModularVehicle/Source/ChaosModularVehicleEngine/Private/ChaosModularVehicle/SuspensionSimModule.cpp
@@ -0,0 +1,272 @@
+// Copyright Epic Games, Inc. All Rights Reserved.
+
+#include "ChaosModularVehicle/SuspensionSimModule.h"
+#include "SimModule/SimModuleTree.h"
+#include "SimModule/TorqueSimModule.h"
+#include "Chaos/PBDSuspensionConstraints.h"
+#include "PhysicsProxy/SuspensionConstraintProxy.h"
+#include "PhysicsProxy/ClusterUnionPhysicsProxy.h"
+#include "PBDRigidsSolver.h"
+#include "VehicleUtility.h"
+#include "PhysicsEngine/PhysicsObjectExternalInterface.h"
+#include "Physics/PhysicsInterfaceCore.h"
+
+
+#if VEHICLE_DEBUGGING_ENABLED
+UE_DISABLE_OPTIMIZATION_SHIP
+#endif
+
+using namespace Chaos;
+
+	FSuspensionSimModule::FSuspensionSimModule(const FSuspensionSettings& Settings)
+		: TSimModuleSettings<FSuspensionSettings>(Settings)
+		, SpringDisplacement(0.f)
+		, LastDisplacement(0.f)
+		, SpringSpeed(0.f)
+	{
+		AccessSetup().MaxLength = FMath::Abs(Settings.MaxRaise + Settings.MaxDrop);
+
+		if (!FModuleFactoryRegister::Get().ContainsFactory(GetSimType()))
+		{
+			static TSharedPtr<FSuspensionFactory> SharedFactory = MakeShared<FSuspensionFactory>();
+			FModuleFactoryRegister::Get().RegisterFactory(GetSimType(), SharedFactory);
+		}
+	}
+	FSuspensionSimModule::~FSuspensionSimModule()
+	{
+	}
+
+	float FSuspensionSimModule::GetSpringLength() const
+	{
+		return  -(Setup().MaxLength - SpringDisplacement);
+	}
+
+	void FSuspensionSimModule::SetSpringLength(float InLength, float WheelRadius)
+	{
+		float DisplacementInput = InLength;
+		DisplacementInput = FMath::Max(0.f, DisplacementInput);
+		SpringDisplacement = Setup().MaxLength - DisplacementInput;
+	}
+
+	void FSuspensionSimModule::GetWorldRaycastLocation(const FTransform& BodyTransform, float WheelRadius, FSpringTrace& OutTrace)
+	{
+		FVector LocalDirection = Setup().SuspensionAxis;
+		FVector Local = GetParentRelativeTransform().GetLocation(); // change to just a vector and GetLocalLocation
+		FVector WorldLocation = BodyTransform.TransformPosition(Local);
+		FVector WorldDirection = BodyTransform.TransformVector(LocalDirection);
+
+		OutTrace.Start = WorldLocation - WorldDirection * (Setup().MaxRaise);
+		OutTrace.End = WorldLocation + WorldDirection * (Setup().MaxDrop + WheelRadius);
+		float TraceLength = OutTrace.Start.Z - OutTrace.End.Z;
+	}
+
+	void FSuspensionSimModule::OnConstruction_External(Chaos::FClusterUnionPhysicsProxy* Proxy)
+	{
+		Chaos::EnsureIsInGameThreadContext();
+		CreateConstraint(Proxy);
+	}
+
+	void FSuspensionSimModule::OnTermination_External()
+	{
+		Chaos::EnsureIsInGameThreadContext();
+		DestroyConstraint();
+	}
+
+	void FSuspensionSimModule::Simulate(float DeltaTime, const FAllInputs& Inputs, FSimModuleTree& VehicleModuleSystem)
+	{
+		{
+			float ForceIntoSurface = 0.0f;
+			if (SpringDisplacement > 0)
+			{
+				float Damping = Setup().SpringDamping;
+				SpringSpeed = (LastDisplacement - SpringDisplacement) / DeltaTime;
+
+				float StiffnessForce = SpringDisplacement * Setup().SpringRate;
+				float DampingForce = SpringSpeed * Damping;
+				float SuspensionForce = StiffnessForce - DampingForce;
+				LastDisplacement = SpringDisplacement;
+
+				if (SuspensionForce > 0)
+				{
+					ForceIntoSurface = SuspensionForce * Setup().SuspensionForceEffect;
+
+					if (!ConstraintHandle.IsValid())
+					{
+						AddLocalForce(Setup().SuspensionAxis * -SuspensionForce, true, false, true, FColor::Green);
+					}
+				}
+			}
+
+			// tell wheels how much they are being pressed into the ground
+			if (SimModuleTree && WheelSimTreeIndex != INVALID_IDX)
+			{
+				if (Chaos::ISimulationModuleBase* Module = SimModuleTree->AccessSimModule(WheelSimTreeIndex))
+				{
+					check(Module->GetSimType() == eSimType::Wheel);
+					Chaos::FWheelBaseInterface* Wheel = static_cast<Chaos::FWheelBaseInterface*>(Module);
+
+					Wheel->SetForceIntoSurface(ForceIntoSurface);
+				}
+		
+			}
+		}
+
+		if (ConstraintHandle.IsValid())
+		{
+			UpdateConstraint();
+		}
+
+	}
+
+	void FSuspensionSimModule::Animate(Chaos::FClusterUnionPhysicsProxy* Proxy)
+	{
+		if (FPBDRigidClusteredParticleHandle* ClusterChild = GetClusterParticle(Proxy))
+		{
+			float CurrentSpringLength = GetSpringLength();
+
+			FVector RestPos = GetInitialParticleTransform().GetTranslation();
+
+			FVector Movement = Setup().SuspensionAxis * (Setup().MaxRaise + CurrentSpringLength);
+			Movement = GetComponentTransform().TransformVector(Movement);
+			FVector NewPos = RestPos - Movement;
+			ClusterChild->ChildToParent().SetTranslation(NewPos); // local frame for module
+		}
+	}
+
+	void FSuspensionSimModule::CreateConstraint(Chaos::FClusterUnionPhysicsProxy* Proxy)
+	{
+		Chaos::EnsureIsInGameThreadContext();
+		const FVector& LocalOffset = GetInitialParticleTransform().GetLocation();
+
+		if (Proxy)
+		{
+			const FPhysicsObjectHandle& PhysicsObject = Proxy->GetPhysicsObjectHandle();
+
+			if (FChaosScene* Scene = static_cast<FChaosScene*>(FPhysicsObjectExternalInterface::GetScene({ &PhysicsObject, 1 })))
+			{
+				FLockedWritePhysicsObjectExternalInterface Interface = FPhysicsObjectExternalInterface::LockWrite(Scene);
+				if (const FGeometryParticle* Particle = Interface->GetParticle(PhysicsObject))
+				{
+					ConstraintHandle = FPhysicsInterface::CreateSuspension(PhysicsObject, LocalOffset);
+
+					if (ConstraintHandle.IsValid())
+					{
+						if (FSuspensionConstraint* Constraint = static_cast<FSuspensionConstraint*>(ConstraintHandle.Constraint))
+						{
+							Constraint->SetHardstopStiffness(1.0f);
+							Constraint->SetSpringStiffness(Setup().SpringRate * 0.25f);
+							Constraint->SetSpringPreload(Setup().SpringPreload);
+							Constraint->SetSpringDamping(Setup().SpringDamping * 5.0f);
+							Constraint->SetMinLength(-Setup().MaxRaise);
+							Constraint->SetMaxLength(Setup().MaxDrop);
+							Constraint->SetAxis(-Setup().SuspensionAxis);
+						}
+					}
+				}
+			}
+		}
+	}
+
+	void FSuspensionSimModule::DestroyConstraint()
+	{
+		Chaos::EnsureIsInGameThreadContext();
+		FPhysicsCommand::ExecuteWrite(ConstraintHandle, [&](const FPhysicsConstraintHandle& Constraint)
+			{
+				FPhysicsInterface::ReleaseConstraint(ConstraintHandle);
+			});
+	}
+
+	void FSuspensionSimModule::UpdateConstraint()
+	{
+		if (Chaos::FSuspensionConstraint* Constraint = static_cast<Chaos::FSuspensionConstraint*>(ConstraintHandle.Constraint))
+		{
+			if (Constraint && Constraint->IsValid())
+			{
+				if (FSuspensionConstraintPhysicsProxy* Proxy = Constraint->GetProxy<FSuspensionConstraintPhysicsProxy>())
+				{
+					Chaos::FPhysicsSolver* Solver = Proxy->GetSolver<Chaos::FPhysicsSolver>();
+					Solver->SetSuspensionTarget(Constraint, TargetPos, ImpactNormal, WheelInContact);
+				}
+			}
+		}
+	}
+
+	void FSuspensionSimModuleData::FillSimState(ISimulationModuleBase* SimModule)
+	{
+		check(SimModule->GetSimType() == eSimType::Suspension);
+		if (FSuspensionSimModule* Sim = static_cast<FSuspensionSimModule*>(SimModule))
+		{
+			Sim->SpringDisplacement = SpringDisplacement;
+			Sim->LastDisplacement = LastDisplacement;
+		}
+	}
+
+	void FSuspensionSimModuleData::FillNetState(const ISimulationModuleBase* SimModule)
+	{
+		check(SimModule->GetSimType() == eSimType::Suspension);
+		if (const FSuspensionSimModule* Sim = static_cast<const FSuspensionSimModule*>(SimModule))
+		{
+			SpringDisplacement = Sim->SpringDisplacement;
+			LastDisplacement = Sim->LastDisplacement;
+		}
+	}
+
+	void FSuspensionSimModuleData::Lerp(const float LerpFactor, const FModuleNetData& Min, const FModuleNetData& Max)
+	{
+		const FSuspensionSimModuleData& MinData = static_cast<const FSuspensionSimModuleData&>(Min);
+		const FSuspensionSimModuleData& MaxData = static_cast<const FSuspensionSimModuleData&>(Max);
+
+		SpringDisplacement = FMath::Lerp(MinData.SpringDisplacement, MaxData.SpringDisplacement, LerpFactor);
+		LastDisplacement = FMath::Lerp(MinData.LastDisplacement, MaxData.LastDisplacement, LerpFactor);
+	}
+
+#if !(UE_BUILD_SHIPPING || UE_BUILD_TEST)
+	FString FSuspensionSimModuleData::ToString() const
+	{
+		return FString::Printf(TEXT("Module:%s SpringDisplacement:%f LastDisplacement:%f"),
+			*DebugString, SpringDisplacement, LastDisplacement);
+	}
+#endif
+
+	void FSuspensionOutputData::FillOutputState(const ISimulationModuleBase* SimModule)
+	{
+		check(SimModule->GetSimType() == eSimType::Suspension);
+
+		FSimOutputData::FillOutputState(SimModule);
+
+		if (const FSuspensionSimModule* Sim = static_cast<const FSuspensionSimModule*>(SimModule))
+		{
+			SpringDisplacement = Sim->SpringDisplacement;
+			SpringDisplacementVector = -Sim->Setup().SuspensionAxis * Sim->SpringDisplacement + Sim->GetAnimationOffset();
+			SpringSpeed = Sim->SpringSpeed;
+			ImpactNormal = Sim->ImpactNormal;
+
+			AnimFlags = EAnimationFlags::AnimatePosition;
+			AnimationLocOffset = SpringDisplacementVector;
+
+		}
+	}
+
+	void FSuspensionOutputData::Lerp(const FSimOutputData& InCurrent, const FSimOutputData& InNext, float Alpha)
+	{
+		FSimOutputData::Lerp(InCurrent, InNext, Alpha);
+
+		const FSuspensionOutputData& Current = static_cast<const FSuspensionOutputData&>(InCurrent);
+		const FSuspensionOutputData& Next = static_cast<const FSuspensionOutputData&>(InNext);
+
+		SpringDisplacement = FMath::Lerp(Current.SpringDisplacement, Next.SpringDisplacement, Alpha);
+		SpringSpeed = FMath::Lerp(Current.SpringSpeed, Next.SpringSpeed, Alpha);
+		ImpactNormal = FMath::Lerp(Current.ImpactNormal, Next.ImpactNormal, Alpha);
+	}
+
+#if !(UE_BUILD_SHIPPING || UE_BUILD_TEST)
+	FString FSuspensionOutputData::ToString()
+	{
+		return FString::Printf(TEXT("%s, SpringDisplacement=%3.3f, SpringSpeed=%3.3f"), *DebugString, SpringDisplacement, SpringSpeed);
+	}
+#endif
+
+
+#if VEHICLE_DEBUGGING_ENABLED
+UE_ENABLE_OPTIMIZATION_SHIP
+#endif
diff --git a/Engine/Plugins/Experimental/ChaosModularVehicle/Source/ChaosModularVehicleEngine/Private/ChaosModularVehicle/VehicleSimSuspensionComponent.cpp b/Engine/Plugins/Experimental/ChaosModularVehicle/Source/ChaosModularVehicleEngine/Private/ChaosModularVehicle/VehicleSimSuspensionComponent.cpp
--- a/Engine/Plugins/Experimental/ChaosModularVehicle/Source/ChaosModularVehicleEngine/Private/ChaosModularVehicle/VehicleSimSuspensionComponent.cpp
+++ b/Engine/Plugins/Experimental/ChaosModularVehicle/Source/ChaosModularVehicleEngine/Private/ChaosModularVehicle/VehicleSimSuspensionComponent.cpp
@@ -2,6 +2,7 @@
 
 #include "ChaosModularVehicle/VehicleSimSuspensionComponent.h"
 #include "SimModule/SimModulesInclude.h"
+#include "ChaosModularVehicle/SuspensionSimModule.h"
 #include "VehicleUtility.h"
 
 
@@ -21,7 +22,7 @@ UVehicleSimSuspensionComponent::UVehicleSimSuspensionComponent()
 Chaos::ISimulationModuleBase* UVehicleSimSuspensionComponent::CreateNewCoreModule() const
 {
 	// use the UE properties to setup the physics state
-	Chaos::FSuspensionSettings Settings;
+	FSuspensionSettings Settings;
 
 	Settings.SuspensionAxis = SuspensionAxis;
 	Settings.MaxRaise = SuspensionMaxRaise;
@@ -33,7 +34,7 @@ Chaos::ISimulationModuleBase* UVehicleSimSuspensionComponent::CreateNewCoreModul
 
 	//Settings.SwaybarEffect = SwaybarEffect;
 
-	Chaos::ISimulationModuleBase* Suspension = new Chaos::FSuspensionSimModule(Settings);
+	Chaos::ISimulationModuleBase* Suspension = new FSuspensionSimModule(Settings);
 	Suspension->SetAnimationEnabled(bAnimationEnabled);
 	return Suspension;
 
diff --git a/Engine/Plugins/Experimental/ChaosModularVehicle/Source/ChaosModularVehicleEngine/Public/ChaosModularVehicle/ModularVehicleBaseComponent.h b/Engine/Plugins/Experimental/ChaosModularVehicle/Source/ChaosModularVehicleEngine/Public/ChaosModularVehicle/ModularVehicleBaseComponent.h
--- a/Engine/Plugins/Experimental/ChaosModularVehicle/Source/ChaosModularVehicleEngine/Public/ChaosModularVehicle/ModularVehicleBaseComponent.h
+++ b/Engine/Plugins/Experimental/ChaosModularVehicle/Source/ChaosModularVehicleEngine/Public/ChaosModularVehicle/ModularVehicleBaseComponent.h
@@ -122,11 +122,6 @@ public:
 
 	void CreateAssociatedSimComponents(UPrimitiveComponent* AttachedComponent, int ParentIndex, int TransformIndex, Chaos::FSimTreeUpdates& TreeUpdatesOut);
 
-	void CreateConstraint(Chaos::ISimulationModuleBase* NewModule);
-	void DestroyConstraint(int ConstraintIndex);
-	void DestroyAllConstraints();
-	void EnableConstraint(int ConstraintIndex, bool bEnabled);
-
 	void PreTickGT(float DeltaTime);
 	void UpdateState(float DeltaTime);
 	TUniquePtr<FModularVehicleAsyncInput> SetCurrentAsyncData(int32 InputIdx, FChaosSimModuleManagerAsyncOutput* CurOutput, FChaosSimModuleManagerAsyncOutput* NextOutput, float Alpha, int32 VehicleManagerTimestamp);
@@ -361,7 +356,6 @@ private:
 
 	int32 NextConstructionIndex = 0;
 
-	TArray<FPhysicsConstraintHandle> ConstraintHandles;
 	int32 ClusteringCount = 0;
 
 	bool bIsLocallyControlled;
diff --git a/Engine/Plugins/Experimental/ChaosModularVehicle/Source/ChaosModularVehicleEngine/Public/ChaosModularVehicle/SuspensionSimModule.h b/Engine/Plugins/Experimental/ChaosModularVehicle/Source/ChaosModularVehicleEngine/Public/ChaosModularVehicle/SuspensionSimModule.h
new file mode 100644
--- /dev/null
+++ b/Engine/Plugins/Experimental/ChaosModularVehicle/Source/ChaosModularVehicleEngine/Public/ChaosModularVehicle/SuspensionSimModule.h
@@ -0,0 +1,163 @@
+// Copyright Epic Games, Inc. All Rights Reserved.
+
+#pragma once
+
+#include "SimModule/SuspensionBaseInterface.h"
+#include "Chaos/ChaosEngineInterface.h"
+
+namespace Chaos
+{
+	struct FAllInputs;
+	class FSimModuleTree;
+	class FClusterUnionPhysicsProxy;
+}
+
+	struct CHAOSMODULARVEHICLEENGINE_API FSuspensionSimModuleData : public Chaos::FModuleNetData
+	{
+#if !(UE_BUILD_SHIPPING || UE_BUILD_TEST)
+		FSuspensionSimModuleData(int NodeArrayIndex, const FString& InDebugString) : Chaos::FModuleNetData(NodeArrayIndex, InDebugString) {}
+#else
+		FSuspensionSimModuleData(int NodeArrayIndex) : FModuleNetData(NodeArrayIndex) {}
+#endif
+
+		virtual Chaos::eSimType GetType() override { return Chaos::eSimType::Suspension; }
+
+		virtual void FillSimState(Chaos::ISimulationModuleBase* SimModule) override;
+
+		virtual void FillNetState(const Chaos::ISimulationModuleBase* SimModule) override;
+
+		virtual void Serialize(FArchive& Ar) override
+		{
+			Ar << SpringDisplacement;
+			Ar << LastDisplacement;
+		}
+
+		virtual void Lerp(const float LerpFactor, const FModuleNetData& Min, const FModuleNetData& Max) override;
+
+#if !(UE_BUILD_SHIPPING || UE_BUILD_TEST)
+		virtual FString ToString() const override;
+#endif
+
+		float SpringDisplacement = 0.0f;
+		float LastDisplacement = 0.0f;
+	};
+
+	struct CHAOSMODULARVEHICLEENGINE_API FSuspensionOutputData : public Chaos::FSimOutputData
+	{
+		virtual FSimOutputData* MakeNewData() override { return FSuspensionOutputData::MakeNew(); }
+		static FSimOutputData* MakeNew() { return new FSuspensionOutputData(); }
+
+		virtual Chaos::eSimType GetType() override { return Chaos::eSimType::Suspension; }
+		virtual void FillOutputState(const Chaos::ISimulationModuleBase* SimModule) override;
+		virtual void Lerp(const FSimOutputData& InCurrent, const Chaos::FSimOutputData& InNext, float Alpha) override;
+
+#if !(UE_BUILD_SHIPPING || UE_BUILD_TEST)
+		virtual FString ToString() override;
+#endif
+
+		float SpringDisplacement;
+		FVector SpringDisplacementVector;
+		float SpringSpeed;
+		FVector ImpactNormal;
+	};
+
+	struct CHAOSMODULARVEHICLEENGINE_API FSuspensionSettings
+	{
+		FSuspensionSettings()
+			: SuspensionAxis(FVector(0.f, 0.f, -1.f))
+			, RestOffset(FVector::ZeroVector)
+			, MaxRaise(5.f)
+			, MaxDrop(5.f)
+			, MaxLength(0.f)
+			, SpringRate(1.f)
+			, SpringPreload(0.5f)
+			, SpringDamping(0.9f)
+			, SuspensionForceEffect(100.0f)
+		{
+
+		}
+
+		FVector SuspensionAxis;		// local axis, direction of suspension force raycast traces
+		FVector RestOffset;
+		float MaxRaise;				// distance [cm]
+		float MaxDrop;				// distance [cm]
+		float MaxLength;			// distance [cm]
+
+		float SpringRate;			// spring constant
+		float SpringPreload;		// Amount of Spring force (independent spring movement)
+		float SpringDamping;		// limit compression/rebound speed
+
+		float SuspensionForceEffect; // force that presses the wheels into the ground - producing grip
+	};
+	
+
+	class FSuspensionFactory : public Chaos::IFactoryModule
+	{
+	public:
+		TSharedPtr<Chaos::FModuleNetData> GenerateNetData(int32 SimArrayIndex) const
+		{
+			return MakeShared<FSuspensionSimModuleData>(
+				SimArrayIndex
+#if !(UE_BUILD_SHIPPING || UE_BUILD_TEST)
+				, TEXT("ConstraintSuspension")
+#endif			
+			);
+		}
+	};
+
+	class CHAOSMODULARVEHICLEENGINE_API FSuspensionSimModule : public Chaos::FSuspensionBaseInterface, public Chaos::TSimModuleSettings<FSuspensionSettings>
+	{
+		friend FSuspensionSimModuleData;
+		friend FSuspensionOutputData;
+
+	public:
+
+		FSuspensionSimModule(const FSuspensionSettings& Settings);
+
+		virtual TSharedPtr<Chaos::FModuleNetData> GenerateNetData(int SimArrayIndex) const override
+		{
+			return MakeShared<FSuspensionSimModuleData>(
+				SimArrayIndex
+#if !(UE_BUILD_SHIPPING || UE_BUILD_TEST)
+				, GetDebugName()
+#endif			
+			);
+		}
+		virtual ~FSuspensionSimModule() override;
+
+
+		virtual Chaos::FSimOutputData* GenerateOutputData() const override
+		{
+			return FSuspensionOutputData::MakeNew();
+		}
+
+		virtual const FString GetDebugName() const { return TEXT("Suspension"); }
+
+		virtual float GetMaxSpringLength() const override { return Setup().MaxLength; }
+		virtual float GetSpringLength() const override;
+		virtual  void SetSpringLength(float InLength, float WheelRadius) override;
+		virtual void GetWorldRaycastLocation(const FTransform& BodyTransform, float WheelRadius, Chaos::FSpringTrace& OutTrace) override;
+
+		virtual void OnConstruction_External(Chaos::FClusterUnionPhysicsProxy* Proxy) override;
+		virtual void OnTermination_External() override;
+
+		virtual void Simulate(float DeltaTime, const Chaos::FAllInputs& Inputs, Chaos::FSimModuleTree& VehicleModuleSystem) override;
+
+		virtual void Animate(Chaos::FClusterUnionPhysicsProxy* Proxy) override;
+
+		const FVector& GetRestLocation() const { return Setup().RestOffset; }
+
+		void UpdateConstraint();
+
+	protected:
+		void CreateConstraint(Chaos::FClusterUnionPhysicsProxy* Proxy);
+		void DestroyConstraint();
+
+	private:
+
+		float SpringDisplacement;
+		float LastDisplacement;
+		float SpringSpeed;
+
+		FPhysicsConstraintHandle ConstraintHandle;
+	};
diff --git a/Engine/Source/Runtime/Experimental/ChaosVehicles/ChaosVehiclesCore/Private/SimModule/ClutchModule.cpp b/Engine/Source/Runtime/Experimental/ChaosVehicles/ChaosVehiclesCore/Private/SimModule/ClutchModule.cpp
--- a/Engine/Source/Runtime/Experimental/ChaosVehicles/ChaosVehiclesCore/Private/SimModule/ClutchModule.cpp
+++ b/Engine/Source/Runtime/Experimental/ChaosVehicles/ChaosVehiclesCore/Private/SimModule/ClutchModule.cpp
@@ -15,7 +15,7 @@ namespace Chaos
 	{
 		if (!FModuleFactoryRegister::Get().ContainsFactory(GetSimType()))
 		{
-			static TSharedPtr<FSimFactoryModule<FClutchSimModuleDatas>> SharedFactory = MakeShared<FSimFactoryModule<FClutchSimModuleDatas>>(GetDebugName());
+			static TSharedPtr<FSimFactoryModule<FClutchSimModuleData>> SharedFactory = MakeShared<FSimFactoryModule<FClutchSimModuleData>>(GetDebugName());
 			FModuleFactoryRegister::Get().RegisterFactory(GetSimType(), SharedFactory);
 		}
 	}
diff --git a/Engine/Source/Runtime/Experimental/ChaosVehicles/ChaosVehiclesCore/Private/SimModule/EngineModule.cpp b/Engine/Source/Runtime/Experimental/ChaosVehicles/ChaosVehiclesCore/Private/SimModule/EngineModule.cpp
--- a/Engine/Source/Runtime/Experimental/ChaosVehicles/ChaosVehiclesCore/Private/SimModule/EngineModule.cpp
+++ b/Engine/Source/Runtime/Experimental/ChaosVehicles/ChaosVehiclesCore/Private/SimModule/EngineModule.cpp
@@ -17,7 +17,7 @@ namespace Chaos
 	{
 		if (!FModuleFactoryRegister::Get().ContainsFactory(GetSimType()))
 		{
-			static TSharedPtr<FSimFactoryModule<FEngineSimModuleDatas>> SharedFactory = MakeShared<FSimFactoryModule<FEngineSimModuleDatas>>(GetDebugName());
+			static TSharedPtr<FSimFactoryModule<FEngineSimModuleData>> SharedFactory = MakeShared<FSimFactoryModule<FEngineSimModuleData>>(GetDebugName());
 			FModuleFactoryRegister::Get().RegisterFactory(GetSimType(), SharedFactory);
 		}
 
diff --git a/Engine/Source/Runtime/Experimental/ChaosVehicles/ChaosVehiclesCore/Private/SimModule/SuspensionBaseInterface.cpp b/Engine/Source/Runtime/Experimental/ChaosVehicles/ChaosVehiclesCore/Private/SimModule/SuspensionBaseInterface.cpp
new file mode 100644
--- /dev/null
+++ b/Engine/Source/Runtime/Experimental/ChaosVehicles/ChaosVehiclesCore/Private/SimModule/SuspensionBaseInterface.cpp
@@ -0,0 +1,36 @@
+// Copyright Epic Games, Inc. All Rights Reserved.
+
+#include "SimModule/SuspensionBaseInterface.h"
+
+namespace Chaos
+{
+
+FSuspensionBaseInterface::FSuspensionBaseInterface()
+	: WheelSimTreeIndex(INVALID_IDX)
+	, TargetPos(FVector::ZeroVector)
+	, ImpactNormal(FVector::ZeroVector)
+	, HitDistance(0.f)
+	, WheelInContact(false)
+{
+}
+
+bool FSuspensionBaseInterface::IsBehaviourType(eSimModuleTypeFlags InType) const
+{ 
+	return (InType & Raycast); 
+}
+
+eSimType FSuspensionBaseInterface::GetSimType() const 
+{ 
+	return eSimType::Suspension; 
+}
+
+void FSuspensionBaseInterface::SetTargetPoint(const FVector& InTargetPoint, const FVector& InImpactNormal, float InHitDistance, bool InWheelInContact)
+{
+	TargetPos = InTargetPoint;
+	ImpactNormal = InImpactNormal;
+	HitDistance = InHitDistance;
+	WheelInContact = InWheelInContact;
+}
+
+
+} // namespace Chaos
diff --git a/Engine/Source/Runtime/Experimental/ChaosVehicles/ChaosVehiclesCore/Private/SimModule/SuspensionModule.cpp b/Engine/Source/Runtime/Experimental/ChaosVehicles/ChaosVehiclesCore/Private/SimModule/SuspensionModule.cpp
deleted file mode 100644
diff --git a/Engine/Source/Runtime/Experimental/ChaosVehicles/ChaosVehiclesCore/Private/SimModule/TorqueSimModule.cpp b/Engine/Source/Runtime/Experimental/ChaosVehicles/ChaosVehiclesCore/Private/SimModule/TorqueSimModule.cpp
--- a/Engine/Source/Runtime/Experimental/ChaosVehicles/ChaosVehiclesCore/Private/SimModule/TorqueSimModule.cpp
+++ b/Engine/Source/Runtime/Experimental/ChaosVehicles/ChaosVehiclesCore/Private/SimModule/TorqueSimModule.cpp
@@ -103,7 +103,7 @@ void FTorqueSimModule::IntegrateAngularVelocity(float DeltaTime, float Inertia,
 		, *GetDebugName(), DriveTorque, BrakingTorque, LoadTorque, AngularVelocity, GetRPM(), AngularPosition, Inertia);
 }
 
-void FTorqueSimModuleDatas::FillSimState(ISimulationModuleBase* SimModule)
+void FTorqueSimModuleData::FillSimState(ISimulationModuleBase* SimModule)
 {
 	if (FTorqueSimModule* Sim = static_cast<FTorqueSimModule*>(SimModule))
 	{
@@ -112,7 +112,7 @@ void FTorqueSimModuleDatas::FillSimState(ISimulationModuleBase* SimModule)
 	}
 }
 
-void FTorqueSimModuleDatas::FillNetState(const ISimulationModuleBase* SimModule)
+void FTorqueSimModuleData::FillNetState(const ISimulationModuleBase* SimModule)
 {
 	if (const FTorqueSimModule* Sim = static_cast<const FTorqueSimModule*>(SimModule))
 	{
@@ -121,10 +121,10 @@ void FTorqueSimModuleDatas::FillNetState(const ISimulationModuleBase* SimModule)
 	}
 }
 
-void FTorqueSimModuleDatas::Lerp(const float LerpFactor, const FModuleNetData& Min, const FModuleNetData& Max)
+void FTorqueSimModuleData::Lerp(const float LerpFactor, const FModuleNetData& Min, const FModuleNetData& Max)
 {
-	const FTorqueSimModuleDatas& MinData = static_cast<const FTorqueSimModuleDatas&>(Min);
-	const FTorqueSimModuleDatas& MaxData = static_cast<const FTorqueSimModuleDatas&>(Max);
+	const FTorqueSimModuleData& MinData = static_cast<const FTorqueSimModuleData&>(Min);
+	const FTorqueSimModuleData& MaxData = static_cast<const FTorqueSimModuleData&>(Max);
 
 	AngularVelocity = FMath::Lerp(MinData.AngularVelocity, MaxData.AngularVelocity, LerpFactor);
 	AngularPosition = FMath::Lerp(MinData.AngularPosition, MaxData.AngularPosition, LerpFactor);
@@ -135,7 +135,7 @@ void FTorqueSimModuleDatas::Lerp(const float LerpFactor, const FModuleNetData& M
 
 
 #if !(UE_BUILD_SHIPPING || UE_BUILD_TEST)
-FString FTorqueSimModuleDatas::ToString() const
+FString FTorqueSimModuleData::ToString() const
 {
 	return FString::Printf(TEXT("Module:%s AngularVelocity:%f AngularPosition:%f"),
 		*DebugString, AngularVelocity, AngularPosition);
diff --git a/Engine/Source/Runtime/Experimental/ChaosVehicles/ChaosVehiclesCore/Private/SimModule/TransmissionModule.cpp b/Engine/Source/Runtime/Experimental/ChaosVehicles/ChaosVehiclesCore/Private/SimModule/TransmissionModule.cpp
--- a/Engine/Source/Runtime/Experimental/ChaosVehicles/ChaosVehiclesCore/Private/SimModule/TransmissionModule.cpp
+++ b/Engine/Source/Runtime/Experimental/ChaosVehicles/ChaosVehiclesCore/Private/SimModule/TransmissionModule.cpp
@@ -20,7 +20,7 @@ namespace Chaos
 	{
 		if (!FModuleFactoryRegister::Get().ContainsFactory(GetSimType()))
 		{
-			static TSharedPtr<FSimFactoryModule<FTransmissionSimModuleDatas>> SharedFactory = MakeShared<FSimFactoryModule<FTransmissionSimModuleDatas>>(GetDebugName());
+			static TSharedPtr<FSimFactoryModule<FTransmissionSimModuleData>> SharedFactory = MakeShared<FSimFactoryModule<FTransmissionSimModuleData>>(GetDebugName());
 			FModuleFactoryRegister::Get().RegisterFactory(GetSimType(), SharedFactory);
 		}
 	}
@@ -168,7 +168,7 @@ namespace Chaos
 		return true;
 	}
 
-	void FTransmissionSimModuleDatas::FillSimState(ISimulationModuleBase* SimModule)
+	void FTransmissionSimModuleData::FillSimState(ISimulationModuleBase* SimModule)
 	{
 		check(SimModule->GetSimType() == eSimType::Transmission);
 		if (FTransmissionSimModule* Sim = static_cast<FTransmissionSimModule*>(SimModule))
@@ -179,7 +179,7 @@ namespace Chaos
 		}
 	}
 
-	void FTransmissionSimModuleDatas::FillNetState(const ISimulationModuleBase* SimModule)
+	void FTransmissionSimModuleData::FillNetState(const ISimulationModuleBase* SimModule)
 	{
 		check(SimModule->GetSimType() == eSimType::Transmission);
 		if (const FTransmissionSimModule* Sim = static_cast<const FTransmissionSimModule*>(SimModule))
@@ -190,10 +190,10 @@ namespace Chaos
 		}
 	}
 
-	void FTransmissionSimModuleDatas::Lerp(const float LerpFactor, const FModuleNetData& Min, const FModuleNetData& Max)
+	void FTransmissionSimModuleData::Lerp(const float LerpFactor, const FModuleNetData& Min, const FModuleNetData& Max)
 	{
-		const FTransmissionSimModuleDatas& MinData = static_cast<const FTransmissionSimModuleDatas&>(Min);
-		const FTransmissionSimModuleDatas& MaxData = static_cast<const FTransmissionSimModuleDatas&>(Max);
+		const FTransmissionSimModuleData& MinData = static_cast<const FTransmissionSimModuleData&>(Min);
+		const FTransmissionSimModuleData& MaxData = static_cast<const FTransmissionSimModuleData&>(Max);
 
 		CurrentGear = LerpFactor < 0.5 ? MinData.CurrentGear : MaxData.CurrentGear;
 		TargetGear = LerpFactor < 0.5 ? MinData.TargetGear : MaxData.TargetGear;
@@ -201,7 +201,7 @@ namespace Chaos
 	}
 
 #if !(UE_BUILD_SHIPPING || UE_BUILD_TEST)
-	FString FTransmissionSimModuleDatas::ToString() const
+	FString FTransmissionSimModuleData::ToString() const
 	{
 		return FString::Printf(TEXT("Module:%s CurrentGear:%d TargetGear:%d CurrentGearChangeTime:%f"),
 			*DebugString, CurrentGear, TargetGear, CurrentGearChangeTime);
diff --git a/Engine/Source/Runtime/Experimental/ChaosVehicles/ChaosVehiclesCore/Private/SimModule/WheelModule.cpp b/Engine/Source/Runtime/Experimental/ChaosVehicles/ChaosVehiclesCore/Private/SimModule/WheelModule.cpp
--- a/Engine/Source/Runtime/Experimental/ChaosVehicles/ChaosVehiclesCore/Private/SimModule/WheelModule.cpp
+++ b/Engine/Source/Runtime/Experimental/ChaosVehicles/ChaosVehiclesCore/Private/SimModule/WheelModule.cpp
@@ -19,7 +19,7 @@ namespace Chaos
 	{
 		if (!FModuleFactoryRegister::Get().ContainsFactory(GetSimType()))
 		{
-			static TSharedPtr<FSimFactoryModule<FWheelSimModuleDatas>> SharedFactory = MakeShared<FSimFactoryModule<FWheelSimModuleDatas>>(GetDebugName());
+			static TSharedPtr<FSimFactoryModule<FWheelSimModuleData>> SharedFactory = MakeShared<FSimFactoryModule<FWheelSimModuleData>>(GetDebugName());
 			FModuleFactoryRegister::Get().RegisterFactory(GetSimType(), SharedFactory);
 		}
 	}
diff --git a/Engine/Source/Runtime/Experimental/ChaosVehicles/ChaosVehiclesCore/Public/SimModule/ClutchModule.h b/Engine/Source/Runtime/Experimental/ChaosVehicles/ChaosVehiclesCore/Public/SimModule/ClutchModule.h
--- a/Engine/Source/Runtime/Experimental/ChaosVehicles/ChaosVehiclesCore/Public/SimModule/ClutchModule.h
+++ b/Engine/Source/Runtime/Experimental/ChaosVehicles/ChaosVehiclesCore/Public/SimModule/ClutchModule.h
@@ -22,12 +22,12 @@ namespace Chaos
 	};
 
 
-	struct CHAOSVEHICLESCORE_API FClutchSimModuleDatas : public FTorqueSimModuleDatas
+	struct CHAOSVEHICLESCORE_API FClutchSimModuleData : public FTorqueSimModuleData
 	{
 #if !(UE_BUILD_SHIPPING || UE_BUILD_TEST)
-		FClutchSimModuleDatas(int NodeArrayIndex, const FString& InDebugString) : FTorqueSimModuleDatas(NodeArrayIndex, InDebugString) {}
+		FClutchSimModuleData(int NodeArrayIndex, const FString& InDebugString) : FTorqueSimModuleData(NodeArrayIndex, InDebugString) {}
 #else
-		FClutchSimModuleDatas(int NodeArrayIndex) : FTorqueSimModuleDatas(NodeArrayIndex) {}
+		FClutchSimModuleData(int NodeArrayIndex) : FTorqueSimModuleData(NodeArrayIndex) {}
 #endif
 
 		virtual eSimType GetType() override { return eSimType::Clutch; }
@@ -35,13 +35,13 @@ namespace Chaos
 		virtual void FillSimState(ISimulationModuleBase* SimModule) override
 		{
 			check(SimModule->GetSimType() == eSimType::Clutch);
-			FTorqueSimModuleDatas::FillSimState(SimModule);
+			FTorqueSimModuleData::FillSimState(SimModule);
 		}
 
 		virtual void FillNetState(const ISimulationModuleBase* SimModule) override
 		{
 			check(SimModule->GetSimType() == eSimType::Clutch);
-			FTorqueSimModuleDatas::FillNetState(SimModule);
+			FTorqueSimModuleData::FillNetState(SimModule);
 		}
 
 	};
@@ -64,7 +64,7 @@ namespace Chaos
 
 		virtual TSharedPtr<FModuleNetData> GenerateNetData(int SimArrayIndex) const
 		{
-			return MakeShared<FClutchSimModuleDatas>(
+			return MakeShared<FClutchSimModuleData>(
 				SimArrayIndex
 #if !(UE_BUILD_SHIPPING || UE_BUILD_TEST)
 				, GetDebugName()
diff --git a/Engine/Source/Runtime/Experimental/ChaosVehicles/ChaosVehiclesCore/Public/SimModule/EngineModule.h b/Engine/Source/Runtime/Experimental/ChaosVehicles/ChaosVehiclesCore/Public/SimModule/EngineModule.h
--- a/Engine/Source/Runtime/Experimental/ChaosVehicles/ChaosVehiclesCore/Public/SimModule/EngineModule.h
+++ b/Engine/Source/Runtime/Experimental/ChaosVehicles/ChaosVehiclesCore/Public/SimModule/EngineModule.h
@@ -10,12 +10,12 @@ namespace Chaos
 	struct FAllInputs;
 	class FSimModuleTree;
 
-	struct CHAOSVEHICLESCORE_API FEngineSimModuleDatas : public FTorqueSimModuleDatas
+	struct CHAOSVEHICLESCORE_API FEngineSimModuleData : public FTorqueSimModuleData
 	{
 #if !(UE_BUILD_SHIPPING || UE_BUILD_TEST)
-		FEngineSimModuleDatas(int NodeArrayIndex, const FString& InDebugString) : FTorqueSimModuleDatas(NodeArrayIndex, InDebugString) {}
+		FEngineSimModuleData(int NodeArrayIndex, const FString& InDebugString) : FTorqueSimModuleData(NodeArrayIndex, InDebugString) {}
 #else
-		FEngineSimModuleDatas(int NodeArrayIndex) : FTorqueSimModuleDatas(NodeArrayIndex) {}
+		FEngineSimModuleData(int NodeArrayIndex) : FTorqueSimModuleData(NodeArrayIndex) {}
 #endif
 
 		virtual eSimType GetType() override { return eSimType::Engine; }
@@ -23,13 +23,13 @@ namespace Chaos
 		virtual void FillSimState(ISimulationModuleBase* SimModule) override
 		{
 			check(SimModule->GetSimType() == eSimType::Engine);
-			FTorqueSimModuleDatas::FillSimState(SimModule);
+			FTorqueSimModuleData::FillSimState(SimModule);
 		}
 
 		virtual void FillNetState(const ISimulationModuleBase* SimModule) override
 		{
 			check(SimModule->GetSimType() == eSimType::Engine);
-			FTorqueSimModuleDatas::FillNetState(SimModule);
+			FTorqueSimModuleData::FillNetState(SimModule);
 		}
 
 	};
@@ -92,7 +92,7 @@ namespace Chaos
 
 		virtual TSharedPtr<FModuleNetData> GenerateNetData(int SimArrayIndex) const
 		{
-			return MakeShared<FEngineSimModuleDatas>(
+			return MakeShared<FEngineSimModuleData>(
 				SimArrayIndex
 #if !(UE_BUILD_SHIPPING || UE_BUILD_TEST)
 				, GetDebugName()
diff --git a/Engine/Source/Runtime/Experimental/ChaosVehicles/ChaosVehiclesCore/Public/SimModule/SimModulesInclude.h b/Engine/Source/Runtime/Experimental/ChaosVehicles/ChaosVehiclesCore/Public/SimModule/SimModulesInclude.h
--- a/Engine/Source/Runtime/Experimental/ChaosVehicles/ChaosVehiclesCore/Public/SimModule/SimModulesInclude.h
+++ b/Engine/Source/Runtime/Experimental/ChaosVehicles/ChaosVehiclesCore/Public/SimModule/SimModulesInclude.h
@@ -8,7 +8,7 @@
 #include "SimModule/MotorModule.h"
 #include "SimModule/AxleModule.h"
 #include "SimModule/EngineModule.h"
-#include "SimModule/SuspensionModule.h"
+#include "SimModule/SuspensionBaseInterface.h"
 #include "SimModule/WheelModule.h"
 #include "SimModule/TransmissionModule.h"
 #include "SimModule/ClutchModule.h"
diff --git a/Engine/Source/Runtime/Experimental/ChaosVehicles/ChaosVehiclesCore/Public/SimModule/SimulationModuleBase.h b/Engine/Source/Runtime/Experimental/ChaosVehicles/ChaosVehiclesCore/Public/SimModule/SimulationModuleBase.h
--- a/Engine/Source/Runtime/Experimental/ChaosVehicles/ChaosVehiclesCore/Public/SimModule/SimulationModuleBase.h
+++ b/Engine/Source/Runtime/Experimental/ChaosVehicles/ChaosVehiclesCore/Public/SimModule/SimulationModuleBase.h
@@ -182,6 +182,16 @@ namespace Chaos
 		*/
 		void SetStateFlags(eSimModuleState StateFlagsIn) { StateFlags = StateFlagsIn; }
 
+		/*
+		* Any post construction initialisation - called from Game Thread
+		*/
+		virtual void OnConstruction_External(Chaos::FClusterUnionPhysicsProxy* Proxy) {}
+
+		/*
+		* Any cleaning up required - called from game thread
+		*/
+		virtual void OnTermination_External() {}
+
 		/**
 		 * The main Simulation function that is called from the physics async callback thread
 		 */
diff --git a/Engine/Source/Runtime/Experimental/ChaosVehicles/ChaosVehiclesCore/Public/SimModule/SuspensionBaseInterface.h b/Engine/Source/Runtime/Experimental/ChaosVehicles/ChaosVehiclesCore/Public/SimModule/SuspensionBaseInterface.h
--- a/Engine/Source/Runtime/Experimental/ChaosVehicles/ChaosVehiclesCore/Public/SimModule/SuspensionBaseInterface.h
+++ b/Engine/Source/Runtime/Experimental/ChaosVehicles/ChaosVehiclesCore/Public/SimModule/SuspensionBaseInterface.h
@@ -30,33 +30,20 @@ struct CHAOSVEHICLESCORE_API FSpringTrace
 class CHAOSVEHICLESCORE_API FSuspensionBaseInterface : public ISimulationModuleBase
 {
 public:
-	FSuspensionBaseInterface()
-		: WheelSimTreeIndex(INVALID_IDX)
-		, TargetPos(FVector::ZeroVector)
-		, ImpactNormal(FVector::ZeroVector)
-		, HitDistance(0.f)
-		, WheelInContact(false)
-	{
-	}
+	FSuspensionBaseInterface();
 
 	virtual ~FSuspensionBaseInterface() {}
 
-	virtual bool IsBehaviourType(eSimModuleTypeFlags InType) const override { return (InType & Raycast); }
+	virtual bool IsBehaviourType(eSimModuleTypeFlags InType) const override;
 
-	virtual eSimType GetSimType() const override { return eSimType::Suspension; }
+	virtual eSimType GetSimType() const override;
 
 	virtual float GetMaxSpringLength() const = 0;
 	virtual float GetSpringLength() const = 0;
 	virtual void SetSpringLength(float InLength, float WheelRadius) = 0;
 	virtual void GetWorldRaycastLocation(const FTransform& BodyTransform, float WheelRadius, FSpringTrace& OutTrace) = 0;
 
-	void SetTargetPoint(const FVector& InTargetPoint, const FVector& InImpactNormal, float InHitDistance, bool InWheelInContact)
-	{
-		TargetPos = InTargetPoint;
-		ImpactNormal = InImpactNormal;
-		HitDistance = InHitDistance;
-		WheelInContact = InWheelInContact;
-	}
+	void SetTargetPoint(const FVector& InTargetPoint, const FVector& InImpactNormal, float InHitDistance, bool InWheelInContact);
 	void SetWheelSimTreeIndex(int WheelTreeIndexIn) { WheelSimTreeIndex = WheelTreeIndexIn; }
 	int GetWheelSimTreeIndex() const { return WheelSimTreeIndex; }
 
diff --git a/Engine/Source/Runtime/Experimental/ChaosVehicles/ChaosVehiclesCore/Public/SimModule/SuspensionModule.h b/Engine/Source/Runtime/Experimental/ChaosVehicles/ChaosVehiclesCore/Public/SimModule/SuspensionModule.h
deleted file mode 100644
diff --git a/Engine/Source/Runtime/Experimental/ChaosVehicles/ChaosVehiclesCore/Public/SimModule/TorqueSimModule.h b/Engine/Source/Runtime/Experimental/ChaosVehicles/ChaosVehiclesCore/Public/SimModule/TorqueSimModule.h
--- a/Engine/Source/Runtime/Experimental/ChaosVehicles/ChaosVehiclesCore/Public/SimModule/TorqueSimModule.h
+++ b/Engine/Source/Runtime/Experimental/ChaosVehicles/ChaosVehiclesCore/Public/SimModule/TorqueSimModule.h
@@ -9,12 +9,12 @@ namespace Chaos
 {
 	class FTorqueSimModule;
 
-	struct CHAOSVEHICLESCORE_API FTorqueSimModuleDatas : public FModuleNetData
+	struct CHAOSVEHICLESCORE_API FTorqueSimModuleData : public FModuleNetData
 	{
 #if !(UE_BUILD_SHIPPING || UE_BUILD_TEST)
-		FTorqueSimModuleDatas(int NodeArrayIndex, const FString& InDebugString) : FModuleNetData(NodeArrayIndex, InDebugString) {}
+		FTorqueSimModuleData(int NodeArrayIndex, const FString& InDebugString) : FModuleNetData(NodeArrayIndex, InDebugString) {}
 #else
-		FTorqueSimModuleDatas(int NodeArrayIndex) : FModuleNetData(NodeArrayIndex) {}
+		FTorqueSimModuleData(int NodeArrayIndex) : FModuleNetData(NodeArrayIndex) {}
 #endif
 		virtual eSimType GetType() override { return eSimType::TorqueSim; }
 
@@ -47,7 +47,7 @@ namespace Chaos
 
 	class CHAOSVEHICLESCORE_API FTorqueSimModule : public ISimulationModuleBase
 	{
-		friend FTorqueSimModuleDatas;
+		friend FTorqueSimModuleData;
 
 	public:
 		FTorqueSimModule()
diff --git a/Engine/Source/Runtime/Experimental/ChaosVehicles/ChaosVehiclesCore/Public/SimModule/TransmissionModule.h b/Engine/Source/Runtime/Experimental/ChaosVehicles/ChaosVehiclesCore/Public/SimModule/TransmissionModule.h
--- a/Engine/Source/Runtime/Experimental/ChaosVehicles/ChaosVehiclesCore/Public/SimModule/TransmissionModule.h
+++ b/Engine/Source/Runtime/Experimental/ChaosVehicles/ChaosVehiclesCore/Public/SimModule/TransmissionModule.h
@@ -10,12 +10,12 @@ namespace Chaos
 	struct FAllInputs;
 	class FSimModuleTree;
 
-	struct CHAOSVEHICLESCORE_API FTransmissionSimModuleDatas : public FModuleNetData
+	struct CHAOSVEHICLESCORE_API FTransmissionSimModuleData : public FModuleNetData
 	{
 #if !(UE_BUILD_SHIPPING || UE_BUILD_TEST)
-		FTransmissionSimModuleDatas(int NodeArrayIndex, const FString& InDebugString) : FModuleNetData(NodeArrayIndex, InDebugString) {}
+		FTransmissionSimModuleData(int NodeArrayIndex, const FString& InDebugString) : FModuleNetData(NodeArrayIndex, InDebugString) {}
 #else
-		FTransmissionSimModuleDatas(int NodeArrayIndex) : FModuleNetData(NodeArrayIndex) {}
+		FTransmissionSimModuleData(int NodeArrayIndex) : FModuleNetData(NodeArrayIndex) {}
 #endif
 
 		virtual eSimType GetType() override { return eSimType::Transmission; }
@@ -100,7 +100,7 @@ namespace Chaos
 
 	class CHAOSVEHICLESCORE_API FTransmissionSimModule : public FTorqueSimModule, public TSimModuleSettings<FTransmissionSettings>
 	{
-		friend FTransmissionSimModuleDatas;
+		friend FTransmissionSimModuleData;
 		friend FTransmissionOutputData;
 
 	public:
@@ -109,7 +109,7 @@ namespace Chaos
 
 		virtual TSharedPtr<FModuleNetData> GenerateNetData(int SimArrayIndex) const
 		{
-			return MakeShared<FTransmissionSimModuleDatas>(
+			return MakeShared<FTransmissionSimModuleData>(
 				SimArrayIndex
 #if !(UE_BUILD_SHIPPING || UE_BUILD_TEST)
 				, GetDebugName()
diff --git a/Engine/Source/Runtime/Experimental/ChaosVehicles/ChaosVehiclesCore/Public/SimModule/WheelModule.h b/Engine/Source/Runtime/Experimental/ChaosVehicles/ChaosVehiclesCore/Public/SimModule/WheelModule.h
--- a/Engine/Source/Runtime/Experimental/ChaosVehicles/ChaosVehiclesCore/Public/SimModule/WheelModule.h
+++ b/Engine/Source/Runtime/Experimental/ChaosVehicles/ChaosVehiclesCore/Public/SimModule/WheelModule.h
@@ -11,12 +11,12 @@ namespace Chaos
 	class FClusterUnionPhysicsProxy;
 	struct FModuleNetData;
 
-	struct CHAOSVEHICLESCORE_API FWheelSimModuleDatas : public FTorqueSimModuleDatas
+	struct CHAOSVEHICLESCORE_API FWheelSimModuleData : public FTorqueSimModuleData
 	{
 #if !(UE_BUILD_SHIPPING || UE_BUILD_TEST)
-		FWheelSimModuleDatas(int NodeArrayIndex, const FString& InDebugString) : FTorqueSimModuleDatas(NodeArrayIndex, InDebugString) {}
+		FWheelSimModuleData(int NodeArrayIndex, const FString& InDebugString) : FTorqueSimModuleData(NodeArrayIndex, InDebugString) {}
 #else
-		FWheelSimModuleDatas(int NodeArrayIndex) : FTorqueSimModuleDatas(NodeArrayIndex) {}
+		FWheelSimModuleData(int NodeArrayIndex) : FTorqueSimModuleData(NodeArrayIndex) {}
 #endif
 
 		virtual eSimType GetType() override { return eSimType::Wheel; }
@@ -24,13 +24,13 @@ namespace Chaos
 		virtual void FillSimState(ISimulationModuleBase* SimModule) override
 		{
 			check(SimModule->GetSimType() == eSimType::Wheel);
-			FTorqueSimModuleDatas::FillSimState(SimModule);
+			FTorqueSimModuleData::FillSimState(SimModule);
 		}
 
 		virtual void FillNetState(const ISimulationModuleBase* SimModule) override
 		{
 			check(SimModule->GetSimType() == eSimType::Wheel);
-			FTorqueSimModuleDatas::FillNetState(SimModule);
+			FTorqueSimModuleData::FillNetState(SimModule);
 		}
 
 	};
@@ -125,7 +125,7 @@ namespace Chaos
 
 		virtual TSharedPtr<FModuleNetData> GenerateNetData(int SimArrayIndex) const
 		{
-			return MakeShared<FWheelSimModuleDatas>(
+			return MakeShared<FWheelSimModuleData>(
 				SimArrayIndex
 #if !(UE_BUILD_SHIPPING || UE_BUILD_TEST)
 				, GetDebugName()
