Adding an input provider to the Modular Vehicle

diff --git a/Engine/Plugins/Experimental/ChaosModularVehicle/Source/ChaosModularVehicleEngine/Private/ChaosModularVehicle/ChaosSimModuleManager.cpp b/Engine/Plugins/Experimental/ChaosModularVehicle/Source/ChaosModularVehicleEngine/Private/ChaosModularVehicle/ChaosSimModuleManager.cpp
--- a/Engine/Plugins/Experimental/ChaosModularVehicle/Source/ChaosModularVehicleEngine/Private/ChaosModularVehicle/ChaosSimModuleManager.cpp
+++ b/Engine/Plugins/Experimental/ChaosModularVehicle/Source/ChaosModularVehicleEngine/Private/ChaosModularVehicle/ChaosSimModuleManager.cpp
@@ -117,6 +117,11 @@ void FChaosSimModuleManager::RegisterCallbacks(UWorld* InWorld)
 	// Set up our async object manager to handle async ticking and marshaling
 	check(AsyncCallback == nullptr);
 	AsyncCallback = Scene.GetSolver()->CreateAndRegisterSimCallbackObject_External<FChaosSimModuleManagerAsyncCallback>();
+
+	if (FNetworkPhysicsCallback* SolverCallback = static_cast<FNetworkPhysicsCallback*>(Scene.GetSolver()->GetRewindCallback()))
+	{
+		SolverCallback->InjectInputsExternal.AddRaw(this, &FChaosSimModuleManager::InjectInputs_External);
+	}
 }
 
 void FChaosSimModuleManager::UnregisterCallbacks()
@@ -193,6 +198,24 @@ void FChaosSimModuleManager::PostUpdate(FChaosScene* PhysScene)
 {
 }
 
+void FChaosSimModuleManager::InjectInputs_External(int32 PhysicsStep, int32 NumSteps)
+{
+	UWorld* World = Scene.GetOwningWorld();
+	if (IsValid(World) == false)
+	{
+		return;
+	}
+	FChaosSimModuleManagerAsyncInput* AsyncInput = AsyncCallback->GetProducerInputData_External();
+	check(AsyncInput);
+	ensure(AsyncInput->World == World);
+
+	for (TWeakObjectPtr<UModularVehicleBaseComponent> Vehicle : CUVehicles)
+	{
+		Vehicle->ProduceInput(PhysicsStep, NumSteps);
+	}
+
+}
+
 
 void FChaosSimModuleManager::ParallelUpdateVehicles(float DeltaSeconds)
 {
diff --git a/Engine/Plugins/Experimental/ChaosModularVehicle/Source/ChaosModularVehicleEngine/Private/ChaosModularVehicle/InputProducer.cpp b/Engine/Plugins/Experimental/ChaosModularVehicle/Source/ChaosModularVehicleEngine/Private/ChaosModularVehicle/InputProducer.cpp
new file mode 100644
--- /dev/null
+++ b/Engine/Plugins/Experimental/ChaosModularVehicle/Source/ChaosModularVehicleEngine/Private/ChaosModularVehicle/InputProducer.cpp
@@ -0,0 +1,132 @@
+// Copyright Epic Games, Inc. All Rights Reserved.
+
+#include "ChaosModularVehicle/InputProducer.h"
+
+// Default input producer
+
+void UVehicleDefaultInputProducer::InitializeContainer(TArray<FModuleInputSetup>& SetupData, FInputNameMap& NameMapOut)
+{
+	MergedInput.Initialize(SetupData, NameMapOut);
+}
+
+void UVehicleDefaultInputProducer::BufferInput(const FInputNameMap& InNameMap, const FName InName, const FModuleInputValue& InValue)
+{
+	//UE_LOG(LogTemp, Warning, TEXT("BufferInput - InName %s Val=%s"), *InName.ToString(), *InValue.ToString());
+
+	if (BufferedCount == 0)
+	{
+		MergedInput.ZeroValues();
+
+		FInputInterface Inputs(InNameMap, MergedInput);
+		Inputs.SetValue(InName, InValue);
+	}
+	else
+	{
+		// inputs are merged here rather than waiting to do the merge in ProduceInput
+		FInputInterface Inputs(InNameMap, MergedInput);
+		Inputs.MergeValue(InName, InValue);
+	}
+
+	// how many times input has been buffered in between being used by ProduceInput
+	BufferedCount++;
+}
+
+void UVehicleDefaultInputProducer::ProduceInput(int32 PhysicsStep, int32 NumSteps, const FInputNameMap& InNameMap, FModuleInputContainer& InOutContainer)
+{
+	//UE_LOG(LogTemp, Warning, TEXT("ProduceInput - PhysicsStep %d, NumSteps %d"), PhysicsStep, NumSteps);
+	//FInputInterface Inputs(InNameMap, MergedInput);
+	//UE_LOG(LogTemp, Warning, TEXT(".. Throttle %s"), *Inputs.GetValue("Throttle").ToString());
+
+	// copy state out
+	InOutContainer = MergedInput;
+	BufferedCount = 0;
+}
+
+
+// Example Playback Input Producer
+
+void UVehiclePlaybackInputProducer::InitializeContainer(TArray<FModuleInputSetup>& SetupData, FInputNameMap& NameMapOut)
+{
+	int32 Seed = 123;
+	FRandomStream Random(Seed);
+	StartStep = 0;
+
+	// Initialize a single container once
+	FModuleInputContainer InputsForFrame;
+	InputsForFrame.Initialize(SetupData, NameMapOut);
+
+	PlaybackBuffer.Reserve(BufferLength);
+	for (int32 I = 0; I < BufferLength; I++)
+	{
+		// copy initialized container into array
+		PlaybackBuffer.Emplace(InputsForFrame);
+
+		// change value of some intputs
+		FInputInterface Inputs(NameMapOut, PlaybackBuffer[I]);
+		Inputs.SetValue("Throttle", Random.FRand());
+		Inputs.SetValue("Steering", 1.0f - 2.0f * Random.FRand());
+	}
+
+}
+
+void UVehiclePlaybackInputProducer::BufferInput(const FInputNameMap& InNameMap, const FName InName, const FModuleInputValue& InValue)
+{
+	// NOP
+}
+
+void UVehiclePlaybackInputProducer::ProduceInput(int32 PhysicsStep, int32 NumSteps, const FInputNameMap& InNameMap, FModuleInputContainer& InOutContainer)
+{
+	if (StartStep == 0)
+	{
+		StartStep = PhysicsStep;
+	}
+
+	int32 UseIndex = PhysicsStep - StartStep;
+	if (UseIndex >=0 && UseIndex <= BufferLength)
+	{
+		FModuleInputContainer& UseContainer = PlaybackBuffer[PhysicsStep - StartStep];
+		InOutContainer = UseContainer;
+	}
+	else
+	{
+		InOutContainer.ZeroValues();
+		StartStep = PhysicsStep;	// restart loop
+	}
+
+}
+
+
+// Example Random Input Producer
+
+void UVehicleRandomInputProducer::InitializeContainer(TArray<FModuleInputSetup>& SetupData, FInputNameMap& NameMapOut)
+{
+	PlaybackContainer.Initialize(SetupData, NameMapOut);
+}
+
+void UVehicleRandomInputProducer::BufferInput(const FInputNameMap& InNameMap, const FName InName, const FModuleInputValue& InValue)
+{
+	// NOP
+}
+
+void UVehicleRandomInputProducer::ProduceInput(int32 PhysicsStep, int32 NumSteps, const FInputNameMap& InNameMap, FModuleInputContainer& InOutContainer)
+{
+	int32 Seed = 123;
+	static FRandomStream Random(Seed);
+
+	// new control settings generated every ChangeInputFrequency number of frames (every frame is too quick)
+	// previous controls are held in the PlaybackContainer between the changes
+	if (PhysicsStep % ChangeInputFrequency == 0)
+	{
+		// clear old input
+		PlaybackContainer.ZeroValues();
+
+		// generate new random input
+		FInputInterface Inputs(InNameMap, PlaybackContainer);
+		Inputs.SetValue("Throttle", Random.FRand());
+		Inputs.SetValue("Steering", 1.0f - 2.0f * Random.FRand());
+	}
+
+	InOutContainer = PlaybackContainer;
+
+}
+
diff --git a/Engine/Plugins/Experimental/ChaosModularVehicle/Source/ChaosModularVehicleEngine/Private/ChaosModularVehicle/ModularVehicleBaseComponent.cpp b/Engine/Plugins/Experimental/ChaosModularVehicle/Source/ChaosModularVehicleEngine/Private/ChaosModularVehicle/ModularVehicleBaseComponent.cpp
--- a/Engine/Plugins/Experimental/ChaosModularVehicle/Source/ChaosModularVehicleEngine/Private/ChaosModularVehicle/ModularVehicleBaseComponent.cpp
+++ b/Engine/Plugins/Experimental/ChaosModularVehicle/Source/ChaosModularVehicleEngine/Private/ChaosModularVehicle/ModularVehicleBaseComponent.cpp
@@ -5,6 +5,7 @@
 #include "ChaosModularVehicle/VehicleSimBaseComponent.h"
 #include "ChaosModularVehicle/ModularVehicleDefaultAsyncInput.h"
 #include "ChaosModularVehicle/ModularVehicleAnimationInstance.h"
+#include "ChaosModularVehicle/InputProducer.h"
 #include "Engine/Engine.h"
 #include "Engine/Canvas.h"
 #include "Engine/OverlapResult.h"
@@ -66,6 +67,8 @@ UModularVehicleBaseComponent::UModularVehicleBaseComponent(const FObjectInitiali
 	}
 
 	bIsLocallyControlled = false;
+
+	InputProducerClass = UVehicleDefaultInputProducer::StaticClass();
 }
 
 UModularVehicleBaseComponent::~UModularVehicleBaseComponent()
@@ -110,43 +113,56 @@ bool UModularVehicleBaseComponent::IsLocallyControlled() const
 	return false;
 }
 
-void UModularVehicleBaseComponent::GenerateInputModifiers(const TArray<FModuleInputSetup>& CombinedInputConfiguration)
+void UModularVehicleBaseComponent::ProduceInput(int32 PhysicsStep, int32 NumSteps)
 {
-	for (const FModuleInputSetup& InputSetup : CombinedInputConfiguration)
+	if (InputProducer)
 	{
-		if (InputSetup.InputModifierClass != nullptr)
-		{
-			UDefaultModularVehicleInputModifier* NewPtr = NewObject<UDefaultModularVehicleInputModifier>(this, InputSetup.InputModifierClass);
-			InputModifiers.Add(NewPtr);
-		}
-		else
-		{
-			InputModifiers.Add(nullptr);
-		}
+		InputProducer->ProduceInput(PhysicsStep, NumSteps, InputNameMap, InputsContainer);
 	}
 }
 
-void UModularVehicleBaseComponent::ApplyInputModifiers(float DeltaTime, const FModuleInputContainer& RawValue)
-{
-	check(InputModifiers.Num() == InputsContainer.GetNumInputs());
-	for (int I = 0; I < InputsContainer.GetNumInputs(); I++)
-	{
-		if (InputModifiers[I] != nullptr)
-		{
-			InputsContainer.SetValueAtIndex(I, InputModifiers[I]->InterpInputValue(DeltaTime, InputsContainer.GetValueAtIndex(I), RawValue.GetValueAtIndex(I)));
-		}
-		else
-		{
-			InputsContainer.SetValueAtIndex(I, RawValue.GetValueAtIndex(I));
-		}
-	}
-}
+//void UModularVehicleBaseComponent::GenerateInputModifiers(const TArray<FModuleInputSetup>& CombinedInputConfiguration)
+//{
+//	for (const FModuleInputSetup& InputSetup : CombinedInputConfiguration)
+//	{
+//		if (InputSetup.InputModifierClass != nullptr)
+//		{
+//			UDefaultModularVehicleInputModifier* NewPtr = NewObject<UDefaultModularVehicleInputModifier>(this, InputSetup.InputModifierClass);
+//			InputModifiers.Add(NewPtr);
+//		}
+//		else
+//		{
+//			InputModifiers.Add(nullptr);
+//		}
+//	}
+//}
+
+//void UModularVehicleBaseComponent::ApplyInputModifiers(float DeltaTime, const FModuleInputContainer& RawValue)
+//{
+//	check(InputModifiers.Num() == InputsContainer.GetNumInputs());
+//	for (int I = 0; I < InputsContainer.GetNumInputs(); I++)
+//	{
+//		if (InputModifiers[I] != nullptr)
+//		{
+//			InputsContainer.SetValueAtIndex(I, InputModifiers[I]->InterpInputValue(DeltaTime, InputsContainer.GetValueAtIndex(I), RawValue.GetValueAtIndex(I)));
+//		}
+//		else
+//		{
+//			InputsContainer.SetValueAtIndex(I, RawValue.GetValueAtIndex(I));
+//		}
+//	}
+//}
 
 
 void UModularVehicleBaseComponent::OnCreatePhysicsState()
 {
 	Super::OnCreatePhysicsState();
 
+	if (InputProducerClass)
+	{
+		InputProducer = InputProducerClass->GetDefaultObject<UVehicleInputProducerBase>();
+	}
+
 	if (ClusterUnionComponent)
 	{
 		// piggyback on the Add/Remove component events from the cluster union to add/remove simulation modules
@@ -368,13 +384,19 @@ void UModularVehicleBaseComponent::BeginPlay()
 	// at that time and AssimilateComponentInputs will not find any controls in the component hierarchy
 	TArray<FModuleInputSetup> CombinedInputConfiguration;
 	AssimilateComponentInputs(CombinedInputConfiguration);
-	RawInputsContainer.Initialize(CombinedInputConfiguration, InputNameMap);
-	InputsContainer = RawInputsContainer;
+
+	if (InputProducer)
+	{
+		InputProducer->InitializeContainer(CombinedInputConfiguration, InputNameMap);
+	}
+
+	InputsContainer.Initialize(CombinedInputConfiguration, InputNameMap);
+
 	if (!bUsingNetworkPhysicsPrediction)
 	{
 		ReplicatedState.Container = InputsContainer;
 	}
-	GenerateInputModifiers(CombinedInputConfiguration);
+	// #TODO reinstate ? GenerateInputModifiers(CombinedInputConfiguration);
 
 	VehicleSimulationPT->SetInputMappings(InputNameMap);
 
@@ -465,7 +487,7 @@ void UModularVehicleBaseComponent::UpdateState(float DeltaTime)
 	// Should we remove input instead of relying on replicated state in that case?
 	if (bProcessLocally && PVehicleOutput)
 	{
-		ApplyInputModifiers(DeltaTime, RawInputsContainer);
+		//ApplyInputModifiers(DeltaTime, RawInputsContainer); #TODO: If we put this back where does it go
 
 		if (!bUsingNetworkPhysicsPrediction)
 		{
@@ -1054,52 +1076,73 @@ void UModularVehicleBaseComponent::AddGeometryCollectionsFromOwnedActor()
 	}
 }
 
+void UModularVehicleBaseComponent::SetInputProducerClass(TSubclassOf<UVehicleInputProducerBase> InInputProducerClass)
+{
+	InputProducerClass = InInputProducerClass;
+}
+
 void UModularVehicleBaseComponent::SetInputBool(const FName Name, const bool Value)
 {
-	FInputInterface Inputs(InputNameMap, RawInputsContainer);
-	Inputs.SetValue(Name, Value);
+	if (InputProducer)
+	{
+		InputProducer->BufferInput(InputNameMap, Name, Value);
+	}
 }
 
 void UModularVehicleBaseComponent::SetInputAxis1D(const FName Name, const double Value)
 {
-	FInputInterface Inputs(InputNameMap, RawInputsContainer);
-	Inputs.SetValue(Name, Value);
+	if (InputProducer)
+	{
+		InputProducer->BufferInput(InputNameMap, Name, Value);
+	}
 }
 
 void UModularVehicleBaseComponent::SetInputAxis2D(const FName Name, const FVector2D Value)
 {
-	FInputInterface Inputs(InputNameMap, RawInputsContainer);
-	Inputs.SetValue(Name, Value);
+	if (InputProducer)
+	{
+		InputProducer->BufferInput(InputNameMap, Name, Value);
+	}
 }
 
 void UModularVehicleBaseComponent::SetInputAxis3D(const FName Name, const FVector Value)
 {
-	FInputInterface Inputs(InputNameMap, RawInputsContainer);
-	Inputs.SetValue(Name, Value);
+	if (InputProducer)
+	{
+		InputProducer->BufferInput(InputNameMap, Name, Value);
+	}
 }
 
 void UModularVehicleBaseComponent::SetInput(const FName& Name, const bool Value)
 {
-	FInputInterface Inputs(InputNameMap, RawInputsContainer);
-	Inputs.SetValue(Name, Value);
+	if (InputProducer)
+	{
+		InputProducer->BufferInput(InputNameMap, Name, Value);
+	}
 }
 
 void UModularVehicleBaseComponent::SetInput(const FName& Name, const double Value)
 {
-	FInputInterface Inputs(InputNameMap, RawInputsContainer);
-	Inputs.SetValue(Name, Value);
+	if (InputProducer)
+	{
+		InputProducer->BufferInput(InputNameMap, Name, Value);
+	}
 }
 
 void UModularVehicleBaseComponent::SetInput(const FName& Name, const FVector2D& Value)
 {
-	FInputInterface Inputs(InputNameMap, RawInputsContainer);
-	Inputs.SetValue(Name, Value);
+	if (InputProducer)
+	{
+		InputProducer->BufferInput(InputNameMap, Name, Value);
+	}
 }
 
 void UModularVehicleBaseComponent::SetInput(const FName& Name, const FVector& Value)
 {
-	FInputInterface Inputs(InputNameMap, RawInputsContainer);
-	Inputs.SetValue(Name, Value);
+	if (InputProducer)
+	{
+		InputProducer->BufferInput(InputNameMap, Name, Value);
+	}
 }
 
 void UModularVehicleBaseComponent::SetReverseInput(bool Reverse)
@@ -1141,12 +1184,11 @@ void UModularVehicleBaseComponent::ShowDebugInfo(AHUD* HUD, UCanvas* Canvas, con
 	// draw input values
 	Canvas->SetDrawColor(FColor::White);
 
-	for (int I = 0; I < RawInputsContainer.GetNumInputs(); I++)
+	for (int I = 0; I < InputsContainer.GetNumInputs(); I++)
 	{
-		float Raw = RawInputsContainer.GetValueAtIndex(I).GetMagnitude();
 		float Interpolated = InputsContainer.GetValueAtIndex(I).GetMagnitude();
 
-		YPos += Canvas->DrawText(RenderFont, FString::Printf(TEXT("%s Raw  (%3.2f) %3.2f"), *InputConfig[I].Name.ToString(), Raw, Interpolated), 4, YPos);
+		YPos += Canvas->DrawText(RenderFont, FString::Printf(TEXT("%s %3.2f"), *InputConfig[I].Name.ToString(), Interpolated), 4, YPos);
 	}
 
 	YPos += 10;
diff --git a/Engine/Plugins/Experimental/ChaosModularVehicle/Source/ChaosModularVehicleEngine/Public/ChaosModularVehicle/ChaosSimModuleManager.h b/Engine/Plugins/Experimental/ChaosModularVehicle/Source/ChaosModularVehicleEngine/Public/ChaosModularVehicle/ChaosSimModuleManager.h
--- a/Engine/Plugins/Experimental/ChaosModularVehicle/Source/ChaosModularVehicleEngine/Public/ChaosModularVehicle/ChaosSimModuleManager.h
+++ b/Engine/Plugins/Experimental/ChaosModularVehicle/Source/ChaosModularVehicleEngine/Public/ChaosModularVehicle/ChaosSimModuleManager.h
@@ -56,6 +56,9 @@ public:
 	/** Post update step */
 	void PostUpdate(FChaosScene* PhysScene);
 
+	/** Called on GT but before Physics thread runs - PT tick rate */
+	void InjectInputs_External(int32 PhysicsStep, int32 NumSteps);
+
 	void ParallelUpdateVehicles(float DeltaSeconds);
 
 	/** Find a vehicle manager from an FPhysScene */
diff --git a/Engine/Plugins/Experimental/ChaosModularVehicle/Source/ChaosModularVehicleEngine/Public/ChaosModularVehicle/InputProducer.h b/Engine/Plugins/Experimental/ChaosModularVehicle/Source/ChaosModularVehicleEngine/Public/ChaosModularVehicle/InputProducer.h
new file mode 100644
--- /dev/null
+++ b/Engine/Plugins/Experimental/ChaosModularVehicle/Source/ChaosModularVehicleEngine/Public/ChaosModularVehicle/InputProducer.h
@@ -0,0 +1,72 @@
+// Copyright Epic Games, Inc. All Rights Reserved.
+
+#pragma once
+
+#include "SimModule/ModuleInput.h"
+#include "InputProducer.generated.h"
+
+/** The default input producer that takes real input from the player and provides it to the simulation */
+UCLASS(BlueprintType, Blueprintable)
+class CHAOSMODULARVEHICLEENGINE_API UVehicleDefaultInputProducer : public UVehicleInputProducerBase
+{
+	GENERATED_BODY()
+
+public:
+
+	/* initialize the input buffer container */
+	virtual void InitializeContainer(TArray<FModuleInputSetup>& SetupData, FInputNameMap& NameMapOut) override;
+
+	/** capture input at game thread frequency */
+	virtual void BufferInput(const FInputNameMap& InNameMap, const FName InName, const FModuleInputValue& InValue) override;
+
+	/** produce input for PT simulation at PT frequency */
+	virtual void ProduceInput(int32 PhysicsStep, int32 NumSteps, const FInputNameMap& InNameMap, FModuleInputContainer& InOutContainer) override;
+
+	FModuleInputContainer MergedInput;
+	uint32 BufferedCount = 0;
+};
+
+
+/** Example input generator, generates random input inot a per frame buffer then replays from the buffer, looping back to the start when the bugger is exhausted */
+UCLASS(BlueprintType, Blueprintable)
+class CHAOSMODULARVEHICLEENGINE_API UVehiclePlaybackInputProducer : public UVehicleInputProducerBase
+{
+	GENERATED_BODY()
+
+public:
+
+	/* initialize the input buffer containers */
+	virtual void InitializeContainer(TArray<FModuleInputSetup>& SetupData, FInputNameMap& NameMapOut) override;
+
+	/** capture input at game thread frequency */
+	virtual void BufferInput(const FInputNameMap& InNameMap, const FName InName, const FModuleInputValue& InValue) override;
+
+	/** produce input for PT simulation at PT frequency */
+	virtual void ProduceInput(int32 PhysicsStep, int32 NumSteps, const FInputNameMap& InNameMap, FModuleInputContainer& InOutContainer) override;
+
+	TArray<FModuleInputContainer> PlaybackBuffer;
+	int32 BufferLength = 150;
+	int32 StartStep = 0;
+};
+
+
+/** Example input generator, generates random input on the fly for the PT */
+UCLASS(BlueprintType, Blueprintable)
+class CHAOSMODULARVEHICLEENGINE_API UVehicleRandomInputProducer : public UVehicleInputProducerBase
+{
+	GENERATED_BODY()
+
+public:
+
+	/* initialize the input buffer containers */
+	virtual void InitializeContainer(TArray<FModuleInputSetup>& SetupData, FInputNameMap& NameMapOut) override;
+
+	/** capture input at game thread frequency */
+	virtual void BufferInput(const FInputNameMap& InNameMap, const FName InName, const FModuleInputValue& InValue) override;
+
+	/** produce input for PT simulation at PT frequency */
+	virtual void ProduceInput(int32 PhysicsStep, int32 NumSteps, const FInputNameMap& InNameMap, FModuleInputContainer& InOutContainer) override;
+
+	FModuleInputContainer PlaybackContainer;
+	int32 ChangeInputFrequency = 10;
+};
\ No newline at end of file
diff --git a/Engine/Plugins/Experimental/ChaosModularVehicle/Source/ChaosModularVehicleEngine/Public/ChaosModularVehicle/ModularVehicleBaseComponent.h b/Engine/Plugins/Experimental/ChaosModularVehicle/Source/ChaosModularVehicleEngine/Public/ChaosModularVehicle/ModularVehicleBaseComponent.h
--- a/Engine/Plugins/Experimental/ChaosModularVehicle/Source/ChaosModularVehicleEngine/Public/ChaosModularVehicle/ModularVehicleBaseComponent.h
+++ b/Engine/Plugins/Experimental/ChaosModularVehicle/Source/ChaosModularVehicleEngine/Public/ChaosModularVehicle/ModularVehicleBaseComponent.h
@@ -5,6 +5,7 @@
 #include "ChaosModularVehicle/ChaosSimModuleManagerAsyncCallback.h"
 #include "SimModule/SimModulesInclude.h"
 #include "SimModule/ModuleInput.h"
+#include "ChaosModularVehicle/InputProducer.h"
 #include "ChaosModularVehicle/ModularVehicleSimulationCU.h"
 #include "Chaos/ParticleHandleFwd.h"
 #include "Components/PrimitiveComponent.h"
@@ -120,6 +121,8 @@ public:
 	virtual void BeginPlay() override;
 	virtual void TickComponent(float DeltaTime, ELevelTick TickType, FActorComponentTickFunction* ThisTickFunction) override;
 
+	void ProduceInput(int32 PhysicsStep, int32 NumSteps);
+
 	void CreateAssociatedSimComponents(UPrimitiveComponent* AttachedComponent, int ParentIndex, int TransformIndex, Chaos::FSimTreeUpdates& TreeUpdatesOut);
 
 	void PreTickGT(float DeltaTime);
@@ -174,13 +177,19 @@ public:
 	UFUNCTION(BlueprintCallable, Category = "Game|Components|ModularVehicle")
 	void SetLocallyControlled(bool bLocallyControlledIn);
 
+	UPROPERTY(EditAnywhere, Category = "Game|Components|ModularVehicle")
+	TSubclassOf<UVehicleInputProducerBase> InputProducerClass;
+
 	// CONTROLS
 	// 
+	//void SetInput(const FName& Name, const FModuleInputValue& Value);
 	void SetInput(const FName& Name, const bool Value);
 	void SetInput(const FName& Name, const double Value);
 	void SetInput(const FName& Name, const FVector2D& Value);
 	void SetInput(const FName& Name, const FVector& Value);
 
+	UFUNCTION(BlueprintCallable, Category = "Game|Components|ModularVehicle")
+	void SetInputProducerClass(TSubclassOf<UVehicleInputProducerBase> InInputProducerClass);
 
 	UFUNCTION(BlueprintCallable, Category = "Game|Components|ModularVehicle")
 	void SetInputBool(const FName Name, const bool Value);
@@ -257,8 +266,8 @@ protected:
 	void AddGeometryCollectionsFromOwnedActor();
 	void SetupSkeletalAnimationStructure();
 	void AssimilateComponentInputs(TArray<FModuleInputSetup>& OutCombinedInputs);
-	virtual void GenerateInputModifiers(const TArray<FModuleInputSetup>& CombinedInputConfiguration);
-	virtual void ApplyInputModifiers(float DeltaTime, const FModuleInputContainer& RawValue);
+	// #TODO reinstate? virtual void GenerateInputModifiers(const TArray<FModuleInputSetup>& CombinedInputConfiguration);
+	// #TODO reinstate? virtual void ApplyInputModifiers(float DeltaTime, const FModuleInputContainer& RawValue);
 
 
 	void ActionTreeUpdates(Chaos::FSimTreeUpdates* NextTreeUpdates);
@@ -363,7 +372,7 @@ private:
 	TArray<FModuleAnimationSetup> ModuleAnimationSetups;
 
 	FInputNameMap InputNameMap;	// map input name to input container array index
-	FModuleInputContainer RawInputsContainer;
+	TObjectPtr<UVehicleInputProducerBase> InputProducer = nullptr;
 	FModuleInputContainer InputsContainer;
 
 };
diff --git a/Engine/Source/Runtime/Experimental/ChaosVehicles/ChaosVehiclesCore/Private/SimModule/ModuleInput.cpp b/Engine/Source/Runtime/Experimental/ChaosVehicles/ChaosVehiclesCore/Private/SimModule/ModuleInput.cpp
--- a/Engine/Source/Runtime/Experimental/ChaosVehicles/ChaosVehiclesCore/Private/SimModule/ModuleInput.cpp
+++ b/Engine/Source/Runtime/Experimental/ChaosVehicles/ChaosVehiclesCore/Private/SimModule/ModuleInput.cpp
@@ -202,6 +202,14 @@ void FModuleInputContainer::Initialize(TArray<FModuleInputSetup>& SetupData, FIn
 	}
 }
 
+void FModuleInputContainer::ZeroValues()
+{
+	for (int I = 0; I < InputValues.Num(); I++)
+	{
+		InputValues[I].Reset();
+	}
+}
+
 void FModuleInputContainer::Serialize(FArchive& Ar, UPackageMap* Map, bool& bOutSuccess)
 {
 	bOutSuccess = true;
@@ -269,6 +277,19 @@ void FInputInterface::SetValue(const FName& InName, const FModuleInputValue& InV
 	}
 }
 
+void FInputInterface::MergeValue(const FName& InName, const FModuleInputValue& InValue)
+{
+	if (const int* Index = NameMap.Find(InName))
+	{
+		ValueContainer.MergeValueAtIndex(*Index, InValue);
+	}
+	else
+	{
+		UE_LOG(LogModularInput, Warning, TEXT("Trying to set the value of an undefined control input %s"), *InName.ToString());
+	}
+}
+
+
 FModuleInputValue FInputInterface::GetValue(const FName& InName) const
 {
 	if (const int* Index = NameMap.Find(InName))
diff --git a/Engine/Source/Runtime/Experimental/ChaosVehicles/ChaosVehiclesCore/Public/SimModule/ModuleInput.h b/Engine/Source/Runtime/Experimental/ChaosVehicles/ChaosVehiclesCore/Public/SimModule/ModuleInput.h
--- a/Engine/Source/Runtime/Experimental/ChaosVehicles/ChaosVehiclesCore/Public/SimModule/ModuleInput.h
+++ b/Engine/Source/Runtime/Experimental/ChaosVehicles/ChaosVehiclesCore/Public/SimModule/ModuleInput.h
@@ -364,6 +364,7 @@ public:
 	int GetNumInputs() const { return InputValues.Num(); }
 	FModuleInputValue GetValueAtIndex(int Index) const { return InputValues[Index]; }
 	void SetValueAtIndex(int Index, const FModuleInputValue& InValue) { InputValues[Index] = InValue; }
+	void MergeValueAtIndex(int Index, const FModuleInputValue& InValue) { InputValues[Index].Merge(InValue); }
 
 	FModuleInputContainer& operator=(const FModuleInputContainer& Other)
 	{
@@ -380,6 +381,8 @@ public:
 
 	void Initialize(TArray<FModuleInputSetup>& SetupData, FInputNameMap& NameMapOut);
 
+	void ZeroValues();
+
 	void Serialize(FArchive& Ar, UPackageMap* Map, bool& bOutSuccess);
 
 	int AddInput(EModuleInputValueType Type, TSubclassOf<UDefaultModularVehicleInputModifier>& InputModifierClass);
@@ -410,6 +413,7 @@ public:
 	}
 
 	void SetValue(const FName& InName, const FModuleInputValue& InValue);
+	void MergeValue(const FName& InName, const FModuleInputValue& InValue);
 	FModuleInputValue GetValue(const FName& InName) const;
 	float GetMagnitude(const FName& InName) const;
 	bool InputsNonZero() const;
@@ -418,3 +422,23 @@ public:
 	FModuleInputContainer& ValueContainer;	// per vehicle instance
 };
 
+UCLASS(Abstract, BlueprintType, Blueprintable, EditInlineNew)
+class CHAOSVEHICLESCORE_API UVehicleInputProducerBase : public UClass
+{
+	GENERATED_BODY()
+	
+public:
+	using FInputNameMap = TMap<FName, int>;
+
+	/* initialize the input buffer container(s) */
+	virtual void InitializeContainer(TArray<FModuleInputSetup>& SetupData, FInputNameMap& NameMapOut) {}
+
+	/** capture input at game thread frequency */
+	virtual void BufferInput(const FInputNameMap& InNameMap, const FName InName, const FModuleInputValue& InValue) {}
+
+	/** produce input for PT simulation at PT frequency */
+	virtual void ProduceInput(int32 PhysicsStep, int32 NumSteps, const FInputNameMap& InNameMap, FModuleInputContainer& InOutContainer) {}
+};
+
+
+
