Modular Vehicle (base component);

Adding an interface to the base component so we can have different types of component because not every component needs to be a primitive component.

-- Code --
ModularVehicleBaseComponent:
CreateAssociatedSimComponents() - Changed the primitive component param to scene component. Using the IVehicleSimBaseComponentInterface instead of UVehicleSimBaseComponent.
ComponentToPhysicsObjects - Changed the primitive component key to scene component.
FindParentsLastSimComponent() - Changed the primitive component param to scene component. Changed return type to int32. Using the IVehicleSimBaseComponentInterface instead of UVehicleSimBaseComponent.
RemoveComponentFromSimulation() - Using the IVehicleSimBaseComponentInterface instead of UVehicleSimBaseComponent.
AssimilateComponentInputs() - Using the IVehicleSimBaseComponentInterface instead of UVehicleSimBaseComponent.

VehicleSimBaseComponent:
IVehicleSimBaseComponentInterface - Created this interface so we can have different types of base components.
UVehicleSimBaseComponent - Added the IVehicleSimBaseComponentInterface implementation.
UVehicleSimBaseSceneComponent - Created this class so we can have scene component based components.

#tests PIE, sandbox, make sure everything still works as expected.

diff --git a/Engine/Plugins/Experimental/ChaosModularVehicle/Source/ChaosModularVehicleEngine/Private/ChaosModularVehicle/ModularVehicleBaseComponent.cpp b/Engine/Plugins/Experimental/ChaosModularVehicle/Source/ChaosModularVehicleEngine/Private/ChaosModularVehicle/ModularVehicleBaseComponent.cpp
--- a/Engine/Plugins/Experimental/ChaosModularVehicle/Source/ChaosModularVehicleEngine/Private/ChaosModularVehicle/ModularVehicleBaseComponent.cpp
+++ b/Engine/Plugins/Experimental/ChaosModularVehicle/Source/ChaosModularVehicleEngine/Private/ChaosModularVehicle/ModularVehicleBaseComponent.cpp
@@ -232,7 +232,7 @@ int GenerateNewGuid()
 }
 
 
-void UModularVehicleBaseComponent::CreateAssociatedSimComponents(UPrimitiveComponent* AttachedComponent, int ParentIndex, int TransformIndex, Chaos::FSimTreeUpdates& TreeUpdatesOut)
+void UModularVehicleBaseComponent::CreateAssociatedSimComponents(USceneComponent* AttachedComponent, int ParentIndex, int TransformIndex, Chaos::FSimTreeUpdates& TreeUpdatesOut)
 {
 	using namespace Chaos;
 	if (AttachedComponent == nullptr || ClusterUnionComponent == nullptr)
@@ -246,14 +246,14 @@ void UModularVehicleBaseComponent::CreateAssociatedSimComponents(UPrimitiveCompo
 
 	ensure(TransformIndex < ChildParticles.Num());
 
-	if (const UVehicleSimBaseComponent* Component = Cast<UVehicleSimBaseComponent>(AttachedComponent))
+	if (IVehicleSimBaseComponentInterface* ComponentInterface = Cast<IVehicleSimBaseComponentInterface>(AttachedComponent))
 	{
 		FTransform ClusterUnionComponentTransform = ClusterUnionComponent->GetComponentTransform();
 		FTransform ComponentTransform = AttachedComponent->GetComponentTransform().GetRelativeTransform(ClusterUnionComponentTransform);
 
 		int TreeIndex = INDEX_NONE;
 
-		Chaos::ISimulationModuleBase* NewModule = Component->CreateNewCoreModule();
+		Chaos::ISimulationModuleBase* NewModule = ComponentInterface->CreateNewCoreModule();
 
 		TUniquePtr<Chaos::FSimModuleTree>& SimModuleTree = VehicleSimulationPT->AccessSimComponentTree();
 
@@ -288,37 +288,37 @@ void UModularVehicleBaseComponent::CreateAssociatedSimComponents(UPrimitiveCompo
 		FTransform ClusterredTransform(FQuat::Identity, InitialTransform.GetLocation());
 		NewModule->SetClusteredTransform(ClusterredTransform);
 
-		if (Component->bAnimationEnabled && (Component->BoneName != NAME_None))
+		const bool bIsAnimationEnabled = ComponentInterface->GetAnimationEnabled();
+		const FName ComponentBoneName = ComponentInterface->GetBoneName();
+		if (bIsAnimationEnabled && (ComponentBoneName != NAME_None))
 		{
 			// if bone already exists then use that (seperate wheel and suspension modules can share same bone)
 
 			int FoundIndex = -1;
 			for (int I = 0; I < ModuleAnimationSetups.Num(); I++)
 			{
-				if (ModuleAnimationSetups[I].BoneName == Component->BoneName)
+				if (ModuleAnimationSetups[I].BoneName == ComponentBoneName)
 				{
 					FoundIndex = I;
 					break;
 				}
 			}
 
+			const FVector& ComponentAnimationOffset = ComponentInterface->GetAnimationOffset();
 			if (FoundIndex != -1)
 			{
-				NewModule->SetAnimationData(Component->BoneName, Component->AnimationOffset, FoundIndex);
+				NewModule->SetAnimationData(ComponentBoneName, ComponentAnimationOffset, FoundIndex);
 			}
 			else
 			{
-				NewModule->SetAnimationData(Component->BoneName, Component->AnimationOffset, ModuleAnimationSetups.Num());
+				NewModule->SetAnimationData(ComponentBoneName, ComponentAnimationOffset, ModuleAnimationSetups.Num());
 				FModuleAnimationSetup AnimSetup(NewModule->GetBoneName());
 				ModuleAnimationSetups.Add(AnimSetup);
 			}
 		}
 
 		// store the tree index in the original sim component
-		if (UVehicleSimBaseComponent* SimComponent = Cast<UVehicleSimBaseComponent>(AttachedComponent))
-		{
-			SimComponent->TreeIndex = TreeIndex;
-		}
+		ComponentInterface->SetTreeIndex(TreeIndex);
 		ParentIndex = TreeIndex;
 
 		if (Chaos::FClusterUnionPhysicsProxy* Proxy = ClusterUnionComponent->GetPhysicsProxy())
@@ -334,9 +334,9 @@ void UModularVehicleBaseComponent::CreateAssociatedSimComponents(UPrimitiveCompo
 	// recurse down tree, converting all SimComponents to proper simulation modules
 	for (USceneComponent* Child : Children)
 	{
-		if (UVehicleSimBaseComponent* ChildSimComponent = Cast<UVehicleSimBaseComponent>(Child))
+		if (IVehicleSimBaseComponentInterface* ChildSimComponent = Cast<IVehicleSimBaseComponentInterface>(Child))
 		{
-			CreateAssociatedSimComponents(ChildSimComponent, ParentIndex, TransformIndex, TreeUpdatesOut);
+			CreateAssociatedSimComponents(Child, ParentIndex, TransformIndex, TreeUpdatesOut);
 		}
 	}
 
@@ -789,18 +789,18 @@ void UModularVehicleBaseComponent::ActionTreeUpdates(Chaos::FSimTreeUpdates* Nex
 		});
 }
 
-int UModularVehicleBaseComponent::FindParentsLastSimComponent(const UPrimitiveComponent* AttachedComponent)
+int32 UModularVehicleBaseComponent::FindParentsLastSimComponent(const USceneComponent* AttachedComponent)
 {
 	if (USceneComponent* AttachParent = AttachedComponent->GetAttachParent())
 	{
 		TArray<USceneComponent*> Children;
 		AttachParent->GetChildrenComponents(false, Children);
 
-		for (int I = Children.Num()-1; I>=0; I--)
+		for (int32 ChildIndex = Children.Num() - 1; ChildIndex >= 0; --ChildIndex)
 		{
-			if (UVehicleSimBaseComponent* ChildSimComponent = Cast<UVehicleSimBaseComponent>(Children[I]))
+			if (IVehicleSimBaseComponentInterface* ChildSimComponent = Cast<IVehicleSimBaseComponentInterface>(Children[ChildIndex]))
 			{
-				return ChildSimComponent->TreeIndex;
+				return ChildSimComponent->GetTreeIndex();
 			}
 		}
 	}
@@ -867,12 +867,12 @@ void UModularVehicleBaseComponent::RemoveComponentFromSimulation(UPrimitiveCompo
 
 		for (USceneComponent* ComponentPart : Components)
 		{
-			if (UVehicleSimBaseComponent* ChangedComponent = Cast<UVehicleSimBaseComponent>(ComponentPart))
+			if (IVehicleSimBaseComponentInterface* ChangedComponent = Cast<IVehicleSimBaseComponentInterface>(ComponentPart))
 			{
-				if (FVehicleComponentData* ComponentData = ComponentToPhysicsObjects.Find(ChangedComponent))
+				if (FVehicleComponentData* ComponentData = ComponentToPhysicsObjects.Find(ComponentPart))
 				{
 					LatestTreeUpdates.RemoveNode(ComponentData->Guid);
-					ComponentToPhysicsObjects.Remove(ChangedComponent);
+					ComponentToPhysicsObjects.Remove(ComponentPart);
 				}
 			}
 		}
@@ -946,14 +946,14 @@ void UModularVehicleBaseComponent::AssimilateComponentInputs(TArray<FModuleInput
 
 			for (USceneComponent* Component : ChildComponents)
 			{
-				if (UVehicleSimBaseComponent* GCComponent = Cast<UVehicleSimBaseComponent>(Component))
+				if (IVehicleSimBaseComponentInterface* GCComponent = Cast<IVehicleSimBaseComponentInterface>(Component))
 				{
 					// don't add duplicates, i.e. 4 wheels could be looking for a single steering input
-					for (FModuleInputSetup& Config : GCComponent->InputConfig)
+					for (FModuleInputSetup& Config : GCComponent->GetInputConfig())
 					{
 						if (OutCombinedInputs.Find(Config) == INDEX_NONE)
 						{
-							OutCombinedInputs.Append(GCComponent->InputConfig);
+							OutCombinedInputs.Append(GCComponent->GetInputConfig());
 						}
 					}
 				}
diff --git a/Engine/Plugins/Experimental/ChaosModularVehicle/Source/ChaosModularVehicleEngine/Private/ChaosModularVehicle/VehicleSimBaseComponent.cpp b/Engine/Plugins/Experimental/ChaosModularVehicle/Source/ChaosModularVehicleEngine/Private/ChaosModularVehicle/VehicleSimBaseComponent.cpp
--- a/Engine/Plugins/Experimental/ChaosModularVehicle/Source/ChaosModularVehicleEngine/Private/ChaosModularVehicle/VehicleSimBaseComponent.cpp
+++ b/Engine/Plugins/Experimental/ChaosModularVehicle/Source/ChaosModularVehicleEngine/Private/ChaosModularVehicle/VehicleSimBaseComponent.cpp
@@ -2,14 +2,7 @@
 
 #include "ChaosModularVehicle/VehicleSimBaseComponent.h"
 
-
-UVehicleSimBaseComponent::UVehicleSimBaseComponent(const FObjectInitializer& ObjectInitializer)
-	: Super(ObjectInitializer)
+void UVehicleSimBaseComponent::SetTreeIndex(const int32 NewValue)
 {
-	BoneName = NAME_None;
-	AnimationOffset = FVector::ZeroVector;
-	bAnimationEnabled = false;
-	AnimationSetupIndex = -1;
-	TreeIndex = -1;
+	TreeIndex = NewValue;
 }
-
diff --git a/Engine/Plugins/Experimental/ChaosModularVehicle/Source/ChaosModularVehicleEngine/Public/ChaosModularVehicle/ModularVehicleBaseComponent.h b/Engine/Plugins/Experimental/ChaosModularVehicle/Source/ChaosModularVehicleEngine/Public/ChaosModularVehicle/ModularVehicleBaseComponent.h
--- a/Engine/Plugins/Experimental/ChaosModularVehicle/Source/ChaosModularVehicleEngine/Public/ChaosModularVehicle/ModularVehicleBaseComponent.h
+++ b/Engine/Plugins/Experimental/ChaosModularVehicle/Source/ChaosModularVehicleEngine/Public/ChaosModularVehicle/ModularVehicleBaseComponent.h
@@ -125,7 +125,7 @@ public:
 
 	void ProduceInput(int32 PhysicsStep, int32 NumSteps);
 
-	void CreateAssociatedSimComponents(UPrimitiveComponent* AttachedComponent, int ParentIndex, int TransformIndex, Chaos::FSimTreeUpdates& TreeUpdatesOut);
+	void CreateAssociatedSimComponents(USceneComponent* AttachedComponent, int ParentIndex, int TransformIndex, Chaos::FSimTreeUpdates& TreeUpdatesOut);
 
 	void PreTickGT(float DeltaTime);
 	void UpdateState(float DeltaTime);
@@ -240,7 +240,7 @@ public:
 	int32 DelayClusteringCount;
 
 	/*** Map simulation component to our vehicle setup data */
-	TMap<TObjectKey<UPrimitiveComponent>, FVehicleComponentData> ComponentToPhysicsObjects;
+	TMap<TObjectKey<USceneComponent>, FVehicleComponentData> ComponentToPhysicsObjects;
 
 	UClusterUnionComponent* ClusterUnionComponent;
 
@@ -275,7 +275,7 @@ protected:
 	void ActionTreeUpdates(Chaos::FSimTreeUpdates* NextTreeUpdates);
 
 	void SetCurrentAsyncDataInternal(FModularVehicleAsyncInput* CurInput, int32 InputIdx, FChaosSimModuleManagerAsyncOutput* CurOutput, FChaosSimModuleManagerAsyncOutput* NextOutput, float Alpha, int32 VehicleManagerTimestamp);
-	int FindParentsLastSimComponent(const UPrimitiveComponent* AttachedComponent);
+	int32 FindParentsLastSimComponent(const USceneComponent* AttachedComponent);
 
 	IPhysicsProxyBase* GetPhysicsProxy() const;
 
diff --git a/Engine/Plugins/Experimental/ChaosModularVehicle/Source/ChaosModularVehicleEngine/Public/ChaosModularVehicle/VehicleSimBaseComponent.h b/Engine/Plugins/Experimental/ChaosModularVehicle/Source/ChaosModularVehicleEngine/Public/ChaosModularVehicle/VehicleSimBaseComponent.h
--- a/Engine/Plugins/Experimental/ChaosModularVehicle/Source/ChaosModularVehicleEngine/Public/ChaosModularVehicle/VehicleSimBaseComponent.h
+++ b/Engine/Plugins/Experimental/ChaosModularVehicle/Source/ChaosModularVehicleEngine/Public/ChaosModularVehicle/VehicleSimBaseComponent.h
@@ -36,33 +36,85 @@ enum class ESimModuleType : uint8
 	Balloon			// TODO: rename anti gravity??
 };
 
+/** Interface used for shared functionality between types of base components. */
+UINTERFACE(MinimalAPI, BlueprintType, meta = (CannotImplementInterfaceInBlueprint))
+class UVehicleSimBaseComponentInterface : public UInterface
+{
+	GENERATED_BODY()
+};
 
-UCLASS(BlueprintType, Blueprintable)
-class CHAOSMODULARVEHICLEENGINE_API UVehicleSimBaseComponent : public UPrimitiveComponent
+class IVehicleSimBaseComponentInterface
 {
-	GENERATED_UCLASS_BODY()
+	GENERATED_BODY()
+
 public:
 
 	virtual ESimModuleType GetModuleType() const { return ESimModuleType::Undefined; }
-
-	/**
-	 * Caller takes ownership of pointer to new Sim Module
-	 */
+	/** Caller takes ownership of pointer to new Sim Module. */
 	virtual Chaos::ISimulationModuleBase* CreateNewCoreModule() const { return nullptr; }
+	virtual FName GetBoneName() const { return NAME_None; }
+	virtual const FVector& GetAnimationOffset() const { return FVector::ZeroVector; }
+	virtual bool GetAnimationEnabled() const { return false; }
+	virtual TArray<FModuleInputSetup> GetInputConfig() const { return TArray<FModuleInputSetup>(); }
+	virtual int32 GetAnimationSetupIndex() const { return INDEX_NONE; }
+	virtual void SetTreeIndex(const int32 NewValue) {}
+	virtual int32 GetTreeIndex() const { return INDEX_NONE; }
+};
+
+/** This if for sim components that need scene component properties along with rendering and collision. */
+UCLASS(BlueprintType, Blueprintable)
+class CHAOSMODULARVEHICLEENGINE_API UVehicleSimBaseComponent 
+	: public UPrimitiveComponent
+	, public IVehicleSimBaseComponentInterface
+{
+	GENERATED_BODY()
+
+protected:
 
 	UPROPERTY(EditAnywhere, BlueprintReadWrite, Category = ModularVehicle)
-	FName BoneName;
+	FName BoneName = NAME_None;
 
 	UPROPERTY(EditAnywhere, BlueprintReadWrite, Category = ModularVehicle)
-	FVector AnimationOffset;
+	FVector AnimationOffset = FVector::ZeroVector;
 
 	UPROPERTY(EditAnywhere, BlueprintReadWrite, Category = ModularVehicle)
-	bool bAnimationEnabled;
+	bool bAnimationEnabled = false;
 
 	UPROPERTY(EditAnywhere, Category = VehicleInput)
 	TArray<FModuleInputSetup> InputConfig;
 
-	int AnimationSetupIndex;
-	int TreeIndex; // helper - since Component->GetAttachChildren doesn't contain any data
+	int32 AnimationSetupIndex = INDEX_NONE;
+	int32 TreeIndex = INDEX_NONE; // helper - since Component->GetAttachChildren doesn't contain any data
+
+public:
+
+	/** IVehicleSimBaseComponentInterface overrides */
+	virtual FName GetBoneName() const override { return BoneName; }
+	virtual const FVector& GetAnimationOffset() const override { return AnimationOffset; }
+	virtual bool GetAnimationEnabled() const override { return bAnimationEnabled; }
+	virtual TArray<FModuleInputSetup> GetInputConfig() const override { return InputConfig; }
+	virtual int32 GetAnimationSetupIndex() const override { return AnimationSetupIndex; }
+	virtual void SetTreeIndex(const int32 NewValue) override;
+	virtual int32 GetTreeIndex() const override { return TreeIndex; }
+	/** END IVehicleSimBaseComponentInterface overrides */
 };
 
+/** This if for sim components that need transform and attachment, no rendering, no collision. */
+UCLASS(BlueprintType, Blueprintable)
+class CHAOSMODULARVEHICLEENGINE_API UVehicleSimBaseSceneComponent 
+	: public USceneComponent
+	, public IVehicleSimBaseComponentInterface
+{
+	GENERATED_BODY()
+
+protected:
+
+	UPROPERTY(EditAnywhere, Category = VehicleInput)
+	TArray<FModuleInputSetup> InputConfig;
+
+public:
+
+	/** IVehicleSimBaseComponentInterface overrides */
+	virtual TArray<FModuleInputSetup> GetInputConfig() const override { return InputConfig; }
+	/** END IVehicleSimBaseComponentInterface overrides */
+};
