Fix input so it resets between physics frames

diff --git a/Engine/Plugins/Experimental/ChaosModularVehicle/Source/ChaosModularVehicleEngine/Private/ChaosModularVehicle/InputProducer.cpp b/Engine/Plugins/Experimental/ChaosModularVehicle/Source/ChaosModularVehicleEngine/Private/ChaosModularVehicle/InputProducer.cpp
--- a/Engine/Plugins/Experimental/ChaosModularVehicle/Source/ChaosModularVehicleEngine/Private/ChaosModularVehicle/InputProducer.cpp
+++ b/Engine/Plugins/Experimental/ChaosModularVehicle/Source/ChaosModularVehicleEngine/Private/ChaosModularVehicle/InputProducer.cpp
@@ -13,22 +13,9 @@ void UVehicleDefaultInputProducer::BufferInput(const FInputNameMap& InNameMap, c
 {
 	//UE_LOG(LogTemp, Warning, TEXT("BufferInput - InName %s Val=%s"), *InName.ToString(), *InValue.ToString());
 
-	if (BufferedCount == 0)
-	{
-		MergedInput.ZeroValues();
-
-		FInputInterface Inputs(InNameMap, MergedInput);
-		Inputs.SetValue(InName, InValue);
-	}
-	else
-	{
-		// inputs are merged here rather than waiting to do the merge in ProduceInput
-		FInputInterface Inputs(InNameMap, MergedInput);
-		Inputs.MergeValue(InName, InValue);
-	}
-
-	// how many times input has been buffered in between being used by ProduceInput
-	BufferedCount++;
+	// inputs are merged here rather than buffered, since they would be merged anyway before use in ProduceInput
+	FInputInterface Inputs(InNameMap, MergedInput);
+	Inputs.MergeValue(InName, InValue);
 }
 
 void UVehicleDefaultInputProducer::ProduceInput(int32 PhysicsStep, int32 NumSteps, const FInputNameMap& InNameMap, FModuleInputContainer& InOutContainer)
@@ -39,7 +26,9 @@ void UVehicleDefaultInputProducer::ProduceInput(int32 PhysicsStep, int32 NumStep
 
 	// copy state out
 	InOutContainer = MergedInput;
-	BufferedCount = 0;
+
+	// reset state for next frame
+	MergedInput.ZeroValues();
 }
 
 
diff --git a/Engine/Plugins/Experimental/ChaosModularVehicle/Source/ChaosModularVehicleEngine/Public/ChaosModularVehicle/InputProducer.h b/Engine/Plugins/Experimental/ChaosModularVehicle/Source/ChaosModularVehicleEngine/Public/ChaosModularVehicle/InputProducer.h
--- a/Engine/Plugins/Experimental/ChaosModularVehicle/Source/ChaosModularVehicleEngine/Public/ChaosModularVehicle/InputProducer.h
+++ b/Engine/Plugins/Experimental/ChaosModularVehicle/Source/ChaosModularVehicleEngine/Public/ChaosModularVehicle/InputProducer.h
@@ -23,7 +23,6 @@ public:
 	virtual void ProduceInput(int32 PhysicsStep, int32 NumSteps, const FInputNameMap& InNameMap, FModuleInputContainer& InOutContainer) override;
 
 	FModuleInputContainer MergedInput;
-	uint32 BufferedCount = 0;
 };
 
 
