Potential crash fix in modular vehicle code

diff --git a/Engine/Plugins/Experimental/ChaosModularVehicle/Source/ChaosModularVehicleEngine/Private/ChaosModularVehicle/ModularVehicleBaseComponent.cpp b/Engine/Plugins/Experimental/ChaosModularVehicle/Source/ChaosModularVehicleEngine/Private/ChaosModularVehicle/ModularVehicleBaseComponent.cpp
--- a/Engine/Plugins/Experimental/ChaosModularVehicle/Source/ChaosModularVehicleEngine/Private/ChaosModularVehicle/ModularVehicleBaseComponent.cpp
+++ b/Engine/Plugins/Experimental/ChaosModularVehicle/Source/ChaosModularVehicleEngine/Private/ChaosModularVehicle/ModularVehicleBaseComponent.cpp
@@ -741,8 +741,12 @@ const FTransform& UModularVehicleBaseComponent::GetComponentTransform() const
 
 void UModularVehicleBaseComponent::ActionTreeUpdates(Chaos::FSimTreeUpdates* NextTreeUpdates)
 {
-	Chaos::FClusterUnionPhysicsProxy* Proxy = ClusterUnionComponent->GetPhysicsProxy();
-	Chaos::FPBDRigidsSolver* Solver = Proxy->GetSolver<Chaos::FPBDRigidsSolver>();
+	if (ClusterUnionComponent)
+	{
+		if (Chaos::FClusterUnionPhysicsProxy* Proxy = ClusterUnionComponent->GetPhysicsProxy())
+		{
+			if (Chaos::FPBDRigidsSolver* Solver = Proxy->GetSolver<Chaos::FPBDRigidsSolver>())
+			{
 	Solver->EnqueueCommandImmediate([Proxy, this, NextTreeUpdates = *NextTreeUpdates]() mutable
 		{
 			if (VehicleSimulationPT)
@@ -788,6 +792,9 @@ void UModularVehicleBaseComponent::ActionTreeUpdates(Chaos::FSimTreeUpdates* Nex
 			}
 		});
 }
+		}
+	}
+}
 
 int32 UModularVehicleBaseComponent::FindParentsLastSimComponent(const USceneComponent* AttachedComponent)
 {
