ModularVehicleBaseComponent:

Make a getter for the input producer.
Move input producer creation to beginplay.
Make the input producer class setter create a producer if none exists.

diff --git a/Engine/Plugins/Experimental/ChaosModularVehicle/Source/ChaosModularVehicleEngine/Private/ChaosModularVehicle/ModularVehicleBaseComponent.cpp b/Engine/Plugins/Experimental/ChaosModularVehicle/Source/ChaosModularVehicleEngine/Private/ChaosModularVehicle/ModularVehicleBaseComponent.cpp
--- a/Engine/Plugins/Experimental/ChaosModularVehicle/Source/ChaosModularVehicleEngine/Private/ChaosModularVehicle/ModularVehicleBaseComponent.cpp
+++ b/Engine/Plugins/Experimental/ChaosModularVehicle/Source/ChaosModularVehicleEngine/Private/ChaosModularVehicle/ModularVehicleBaseComponent.cpp
@@ -158,11 +158,6 @@ void UModularVehicleBaseComponent::OnCreatePhysicsState()
 {
 	Super::OnCreatePhysicsState();
 
-	if (InputProducerClass)
-	{
-		InputProducer = NewObject<UVehicleInputProducerBase>(this, InputProducerClass);
-	}
-
 	if (ClusterUnionComponent)
 	{
 		// piggyback on the Add/Remove component events from the cluster union to add/remove simulation modules
@@ -385,6 +380,11 @@ void UModularVehicleBaseComponent::BeginPlay()
 	TArray<FModuleInputSetup> CombinedInputConfiguration;
 	AssimilateComponentInputs(CombinedInputConfiguration);
 
+	if (!InputProducer && InputProducerClass)
+	{
+		InputProducer = NewObject<UVehicleInputProducerBase>(this, InputProducerClass);
+	}
+
 	if (InputProducer)
 	{
 		InputProducer->InitializeContainer(CombinedInputConfiguration, InputNameMap);
@@ -1086,6 +1086,10 @@ void UModularVehicleBaseComponent::AddGeometryCollectionsFromOwnedActor()
 void UModularVehicleBaseComponent::SetInputProducerClass(TSubclassOf<UVehicleInputProducerBase> InInputProducerClass)
 {
 	InputProducerClass = InInputProducerClass;
+	if (!InputProducer)
+	{
+		InputProducer = NewObject<UVehicleInputProducerBase>(this, InputProducerClass);
+	}
 }
 
 void UModularVehicleBaseComponent::SetInputBool(const FName Name, const bool Value)
diff --git a/Engine/Plugins/Experimental/ChaosModularVehicle/Source/ChaosModularVehicleEngine/Public/ChaosModularVehicle/ModularVehicleBaseComponent.h b/Engine/Plugins/Experimental/ChaosModularVehicle/Source/ChaosModularVehicleEngine/Public/ChaosModularVehicle/ModularVehicleBaseComponent.h
--- a/Engine/Plugins/Experimental/ChaosModularVehicle/Source/ChaosModularVehicleEngine/Public/ChaosModularVehicle/ModularVehicleBaseComponent.h
+++ b/Engine/Plugins/Experimental/ChaosModularVehicle/Source/ChaosModularVehicleEngine/Public/ChaosModularVehicle/ModularVehicleBaseComponent.h
@@ -190,9 +190,12 @@ public:
 	void SetInput(const FName& Name, const FVector2D& Value);
 	void SetInput(const FName& Name, const FVector& Value);
 
+	// Sets the input producer class and creates an input producer if one doesn't exist.
 	UFUNCTION(BlueprintCallable, Category = "Game|Components|ModularVehicle")
 	void SetInputProducerClass(TSubclassOf<UVehicleInputProducerBase> InInputProducerClass);
 
+	TObjectPtr<UVehicleInputProducerBase> GetInputProducer() { return InputProducer; }
+
 	UFUNCTION(BlueprintCallable, Category = "Game|Components|ModularVehicle")
 	void SetInputBool(const FName Name, const bool Value);
 
