Chaos Modular Vehicles: Enable contact modification callback

diff --git a/Engine/Plugins/Experimental/ChaosModularVehicle/Source/ChaosModularVehicleEngine/Private/ChaosModularVehicle/ChaosSimModuleManagerAsyncCallback.cpp b/Engine/Plugins/Experimental/ChaosModularVehicle/Source/ChaosModularVehicleEngine/Private/ChaosModularVehicle/ChaosSimModuleManagerAsyncCallback.cpp
--- a/Engine/Plugins/Experimental/ChaosModularVehicle/Source/ChaosModularVehicleEngine/Private/ChaosModularVehicle/ChaosSimModuleManagerAsyncCallback.cpp
+++ b/Engine/Plugins/Experimental/ChaosModularVehicle/Source/ChaosModularVehicleEngine/Private/ChaosModularVehicle/ChaosSimModuleManagerAsyncCallback.cpp
@@ -13,6 +13,7 @@
 FSimModuleDebugParams GSimModuleDebugParams;
 
 DECLARE_CYCLE_STAT(TEXT("AsyncCallback:OnPreSimulate_Internal"), STAT_AsyncCallback_OnPreSimulate, STATGROUP_ChaosSimModuleManager);
+DECLARE_CYCLE_STAT(TEXT("AsyncCallback:OnContactModification_Internal"), STAT_AsyncCallback_OnContactModification, STATGROUP_ChaosSimModuleManager);
 
 FName FChaosSimModuleManagerAsyncCallback::GetFNameForStatId() const
 {
@@ -111,7 +112,53 @@ void FChaosSimModuleManagerAsyncCallback::OnPreSimulate_Internal()
  */
 void FChaosSimModuleManagerAsyncCallback::OnContactModification_Internal(Chaos::FCollisionContactModifier& Modifications)
 {
+	using namespace Chaos;
+
+	SCOPE_CYCLE_COUNTER(STAT_AsyncCallback_OnContactModification);
+
+	float DeltaTime = GetDeltaTime_Internal();
+	float SimTime = GetSimTime_Internal();
+
+	const FChaosSimModuleManagerAsyncInput* Input = GetConsumerInput_Internal();
+	if (Input == nullptr)
+	{
+		return;
+	}
+
+	const int32 NumVehicles = Input->VehicleInputs.Num();
+
+	UWorld* World = Input->World.Get();	//only safe to access for scene queries
+	if (World == nullptr || NumVehicles == 0)
+	{
+		//world is gone so don't bother.
+		return;
+	}
 
+	Chaos::FPhysicsSolver* PhysicsSolver = static_cast<Chaos::FPhysicsSolver*>(GetSolver());
+	if (PhysicsSolver == nullptr)
+	{
+		return;
+	}
+
+	const TArray<TUniquePtr<FModularVehicleAsyncInput>>& InputVehiclesBatch = Input->VehicleInputs;
+
+	// beware running the vehicle simulation in parallel, code must remain threadsafe
+	auto LambdaParallelUpdate = [&Modifications, &InputVehiclesBatch](int32 Idx)
+	{
+		const FModularVehicleAsyncInput& VehicleInput = *InputVehiclesBatch[Idx];
+
+		if (VehicleInput.Proxy == nullptr)
+		{
+			return;
+		}
+
+		bool bWake = false;
+		VehicleInput.OnContactModification(Modifications);
+
+	};
+
+	bool ForceSingleThread = !GSimModuleDebugParams.EnableMultithreading;
+	PhysicsParallelFor(InputVehiclesBatch.Num(), LambdaParallelUpdate, ForceSingleThread);
 }
 
 
@@ -140,6 +187,14 @@ TUniquePtr<FModularVehicleAsyncOutput> FModularVehicleAsyncInput::Simulate(UWorl
 	return MoveTemp(Output);
 }
 
+void FModularVehicleAsyncInput::OnContactModification(Chaos::FCollisionContactModifier& Modifications) const
+{
+	if (Vehicle && Vehicle->VehicleSimulationPT)
+	{
+		Vehicle->VehicleSimulationPT->OnContactModification(Modifications, Proxy);
+	}
+}
+
 void FModularVehicleAsyncInput::ApplyDeferredForces() const
 {
 	if (Vehicle && Proxy)
diff --git a/Engine/Plugins/Experimental/ChaosModularVehicle/Source/ChaosModularVehicleEngine/Private/ChaosModularVehicle/ModularVehicleSimulationCU.cpp b/Engine/Plugins/Experimental/ChaosModularVehicle/Source/ChaosModularVehicleEngine/Private/ChaosModularVehicle/ModularVehicleSimulationCU.cpp
--- a/Engine/Plugins/Experimental/ChaosModularVehicle/Source/ChaosModularVehicleEngine/Private/ChaosModularVehicle/ModularVehicleSimulationCU.cpp
+++ b/Engine/Plugins/Experimental/ChaosModularVehicle/Source/ChaosModularVehicleEngine/Private/ChaosModularVehicle/ModularVehicleSimulationCU.cpp
@@ -89,6 +89,23 @@ void FModularVehicleSimulationCU::Simulate(UWorld* InWorld, float DeltaSeconds,
 	Simulate_ClusterUnion(InWorld, DeltaSeconds, InputData, OutputData, static_cast<Chaos::FClusterUnionPhysicsProxy*>(Proxy));
 }
 
+void FModularVehicleSimulationCU::OnContactModification(Chaos::FCollisionContactModifier& Modifier, IPhysicsProxyBase* Proxy)
+{
+	using namespace Chaos;
+	Chaos::EnsureIsInPhysicsThreadContext();
+	check(Proxy->GetType() == EPhysicsProxyType::ClusterUnionProxy)
+	Chaos::FClusterUnionPhysicsProxy* ClusterProxy = static_cast<Chaos::FClusterUnionPhysicsProxy*>(Proxy);
+	check(ClusterProxy);
+
+	FPBDRigidsEvolutionGBF& Evolution = *static_cast<FPBDRigidsSolver*>(Proxy->GetSolver<FPBDRigidsSolver>())->GetEvolution();
+	FClusterUnionManager& ClusterUnionManager = Evolution.GetRigidClustering().GetClusterUnionManager();
+	const FClusterUnionIndex& CUI = ClusterProxy->GetClusterUnionIndex();
+
+	if (SimModuleTree.IsValid())
+	{
+		SimModuleTree->OnContactModification(Modifier, ClusterProxy);
+	}
+}
 
 void FModularVehicleSimulationCU::Simulate_ClusterUnion(UWorld* InWorld, float DeltaSeconds, const FModularVehicleAsyncInput& InputData, FModularVehicleAsyncOutput& OutputData, Chaos::FClusterUnionPhysicsProxy* Proxy)
 {
diff --git a/Engine/Plugins/Experimental/ChaosModularVehicle/Source/ChaosModularVehicleEngine/Public/ChaosModularVehicle/ChaosSimModuleManagerAsyncCallback.h b/Engine/Plugins/Experimental/ChaosModularVehicle/Source/ChaosModularVehicleEngine/Public/ChaosModularVehicle/ChaosSimModuleManagerAsyncCallback.h
--- a/Engine/Plugins/Experimental/ChaosModularVehicle/Source/ChaosModularVehicleEngine/Public/ChaosModularVehicle/ChaosSimModuleManagerAsyncCallback.h
+++ b/Engine/Plugins/Experimental/ChaosModularVehicle/Source/ChaosModularVehicleEngine/Public/ChaosModularVehicle/ChaosSimModuleManagerAsyncCallback.h
@@ -217,6 +217,8 @@ struct CHAOSMODULARVEHICLEENGINE_API FModularVehicleAsyncInput
 	* Vehicle simulation running on the Physics Thread
 	*/
 	virtual TUniquePtr<struct FModularVehicleAsyncOutput> Simulate(UWorld* World, const float DeltaSeconds, const float TotalSeconds, bool& bWakeOut) const;
+
+	virtual void OnContactModification(Chaos::FCollisionContactModifier& Modifier) const;
 	virtual void ApplyDeferredForces() const;
 	virtual void ProcessInputs();
 
diff --git a/Engine/Plugins/Experimental/ChaosModularVehicle/Source/ChaosModularVehicleEngine/Public/ChaosModularVehicle/ModularVehicleSimulationCU.h b/Engine/Plugins/Experimental/ChaosModularVehicle/Source/ChaosModularVehicleEngine/Public/ChaosModularVehicle/ModularVehicleSimulationCU.h
--- a/Engine/Plugins/Experimental/ChaosModularVehicle/Source/ChaosModularVehicleEngine/Public/ChaosModularVehicle/ModularVehicleSimulationCU.h
+++ b/Engine/Plugins/Experimental/ChaosModularVehicle/Source/ChaosModularVehicleEngine/Public/ChaosModularVehicle/ModularVehicleSimulationCU.h
@@ -59,6 +59,8 @@ public:
 	virtual void Simulate(UWorld* InWorld, float DeltaSeconds, const FModularVehicleAsyncInput& InputData, FModularVehicleAsyncOutput& OutputData, IPhysicsProxyBase* Proxy);
 	virtual void Simulate_ClusterUnion(UWorld* InWorld, float DeltaSeconds, const FModularVehicleAsyncInput& InputData, FModularVehicleAsyncOutput& OutputData, Chaos::FClusterUnionPhysicsProxy* Proxy);
 
+	virtual void OnContactModification(Chaos::FCollisionContactModifier& Modifier, IPhysicsProxyBase* Proxy);
+
 	void ApplyDeferredForces(FGeometryCollectionPhysicsProxy* RigidHandle);
 	void ApplyDeferredForces(Chaos::FClusterUnionPhysicsProxy* Proxy);
 
diff --git a/Engine/Source/Runtime/Experimental/ChaosVehicles/ChaosVehiclesCore/Private/SimModule/SimModuleTree.cpp b/Engine/Source/Runtime/Experimental/ChaosVehicles/ChaosVehiclesCore/Private/SimModule/SimModuleTree.cpp
--- a/Engine/Source/Runtime/Experimental/ChaosVehicles/ChaosVehiclesCore/Private/SimModule/SimModuleTree.cpp
+++ b/Engine/Source/Runtime/Experimental/ChaosVehicles/ChaosVehiclesCore/Private/SimModule/SimModuleTree.cpp
@@ -252,6 +252,16 @@ void FSimModuleTree::Simulate(float DeltaTime, FAllInputs& Inputs, FClusterUnion
 
 }
 
+void FSimModuleTree::OnContactModification(FCollisionContactModifier& Modifier, FClusterUnionPhysicsProxy* PhysicsProxy)
+{
+	TArray<int> RootNodes;
+	GetRootNodes(RootNodes);
+	for (int RootIndex : RootNodes)
+	{
+		OnContactModificationInternal(RootIndex, Modifier, PhysicsProxy);
+	}
+}
+
 void FSimModuleTree::SimulateNode(float DeltaTime, FAllInputs& Inputs, int NodeIndex, FClusterUnionPhysicsProxy* PhysicsProxy)
 {
 	if (ISimulationModuleBase* Module = AccessSimModule(NodeIndex))
@@ -341,6 +351,18 @@ void FSimModuleTree::SimulateNodeBFS(float DeltaTime, FAllInputs& Inputs, const
 	}
 }
 
+void FSimModuleTree::OnContactModificationInternal(int NodeIndex, FCollisionContactModifier& Modifier, FClusterUnionPhysicsProxy* PhysicsProxy)
+{
+	if (ISimulationModuleBase* Module = AccessSimModule(NodeIndex))
+	{
+		Module->OnContactModification(Modifier, PhysicsProxy);
+	}
+	for (int ChildIdx : GetChildren(NodeIndex))
+	{
+		OnContactModificationInternal(ChildIdx, Modifier, PhysicsProxy);
+	}
+}
+
 
 void FSimModuleTree::DeleteNodesBelow(int AtIndex)
 {
diff --git a/Engine/Source/Runtime/Experimental/ChaosVehicles/ChaosVehiclesCore/Public/SimModule/SimModuleTree.h b/Engine/Source/Runtime/Experimental/ChaosVehicles/ChaosVehiclesCore/Public/SimModule/SimModuleTree.h
--- a/Engine/Source/Runtime/Experimental/ChaosVehicles/ChaosVehiclesCore/Public/SimModule/SimModuleTree.h
+++ b/Engine/Source/Runtime/Experimental/ChaosVehicles/ChaosVehiclesCore/Public/SimModule/SimModuleTree.h
@@ -153,6 +153,8 @@ namespace Chaos
 
 		void Simulate(float DeltaTime, FAllInputs& Inputs, FClusterUnionPhysicsProxy* PhysicsProxy);
 
+		void OnContactModification(FCollisionContactModifier& Modifier, FClusterUnionPhysicsProxy* PhysicsProxy);
+
 		void SetSimTreeProcessingOrder(ESimTreeProcessingOrder OrderIn) { SimTreeProcessingOrder = OrderIn; }
 
 		FDeferredForcesModular& AccessDeferredForces() { return DeferredForces; }
@@ -212,6 +214,8 @@ namespace Chaos
 	protected:
 		void SimulateNode(float DeltaTime, FAllInputs& Inputs, int NodeIdx, FClusterUnionPhysicsProxy* PhysicsProxy);
 
+		void OnContactModificationInternal(int NodeIndex, FCollisionContactModifier& Modifier, FClusterUnionPhysicsProxy* PhysicsProxy);
+
 		void SimulateNodeBFS(float DeltaTime, FAllInputs& Inputs, const TArray<int>& RootNodes, FClusterUnionPhysicsProxy* PhysicsProxy);
 
 		void DeleteNodesBelow(int NodeIdx);
diff --git a/Engine/Source/Runtime/Experimental/ChaosVehicles/ChaosVehiclesCore/Public/SimModule/SimulationModuleBase.h b/Engine/Source/Runtime/Experimental/ChaosVehicles/ChaosVehiclesCore/Public/SimModule/SimulationModuleBase.h
--- a/Engine/Source/Runtime/Experimental/ChaosVehicles/ChaosVehiclesCore/Public/SimModule/SimulationModuleBase.h
+++ b/Engine/Source/Runtime/Experimental/ChaosVehicles/ChaosVehiclesCore/Public/SimModule/SimulationModuleBase.h
@@ -26,6 +26,7 @@ namespace Chaos
 	struct FModuleNetData;
 	struct FSimOutputData;
 	class FClusterUnionPhysicsProxy;
+	class FCollisionContactModifier;
 
 	const FName HandbrakeControlName("Handbrake");
 	const FName ThrottleControlName("Throttle");
@@ -207,6 +208,8 @@ namespace Chaos
 			Simulate(DeltaTime, Inputs, VehicleModuleSystem);
 		}
 
+		virtual void OnContactModification(Chaos::FCollisionContactModifier& Modifier, Chaos::FClusterUnionPhysicsProxy* Proxy) {}
+
 		/**
 		* The main Simulation function that is called from the physics async callback thread
 		*/
