[ModularVehicles] Full rework of typing information and factory module setups.

- Supports a number of useful features, but primarily replaces the content coupling mechanics with a strongly typed grouping enforceable by the compiler.
- Implements helpful features for casting to only allow casting within the right type hierarchy
- Supports single inheritence paranting and type checking complete with casting support to inermediate types.
- Entirely data driven and requires no engine level modifications to support extensible structures.
- Implements self regestering types.

diff --git a/Engine/Plugins/Experimental/ChaosModularVehicle/Source/ChaosModularVehicleEngine/Private/ChaosModularVehicle/ChaosSimModuleManager.cpp b/Engine/Plugins/Experimental/ChaosModularVehicle/Source/ChaosModularVehicleEngine/Private/ChaosModularVehicle/ChaosSimModuleManager.cpp
--- a/Engine/Plugins/Experimental/ChaosModularVehicle/Source/ChaosModularVehicleEngine/Private/ChaosModularVehicle/ChaosSimModuleManager.cpp
+++ b/Engine/Plugins/Experimental/ChaosModularVehicle/Source/ChaosModularVehicleEngine/Private/ChaosModularVehicle/ChaosSimModuleManager.cpp
@@ -70,9 +70,6 @@ void FChaosSimModuleManager::OnWorldCleanup(UWorld* InWorld, bool bSessionEnded,
 	{
 		Manager->UnregisterCallbacks();
 	}
-
-	Chaos::FModuleFactoryRegister::Get().Reset();
-
 }
 
 void FChaosSimModuleManager::OnShowDebugInfo(AHUD* HUD, UCanvas* Canvas, const FDebugDisplayInfo& DisplayInfo, float& YL, float& YPos)
diff --git a/Engine/Plugins/Experimental/ChaosModularVehicle/Source/ChaosModularVehicleEngine/Private/ChaosModularVehicle/ChaosSimModuleManagerAsyncCallback.cpp b/Engine/Plugins/Experimental/ChaosModularVehicle/Source/ChaosModularVehicleEngine/Private/ChaosModularVehicle/ChaosSimModuleManagerAsyncCallback.cpp
--- a/Engine/Plugins/Experimental/ChaosModularVehicle/Source/ChaosModularVehicleEngine/Private/ChaosModularVehicle/ChaosSimModuleManagerAsyncCallback.cpp
+++ b/Engine/Plugins/Experimental/ChaosModularVehicle/Source/ChaosModularVehicleEngine/Private/ChaosModularVehicle/ChaosSimModuleManagerAsyncCallback.cpp
@@ -261,20 +261,26 @@ void FNetworkModularVehicleInputs::ApplyData(UActorComponent* NetworkComponent)
 {
 	if (GSimModuleDebugParams.EnableNetworkStateData)
 	{
-		if (FModularVehicleSimulationCU* VehicleSimulation = Cast<UModularVehicleBaseComponent>(NetworkComponent)->VehicleSimulationPT.Get())
+		if (UModularVehicleBaseComponent* ModularBaseComponent = Cast<UModularVehicleBaseComponent>(NetworkComponent))
 		{
-			VehicleSimulation->VehicleInputs = VehicleInputs;
+			if (FModularVehicleSimulationCU* VehicleSimulation = ModularBaseComponent->VehicleSimulationPT.Get())
+			{
+				VehicleSimulation->VehicleInputs = VehicleInputs;
+			}
 		}
 	}
 }
 
 void FNetworkModularVehicleInputs::BuildData(const UActorComponent* NetworkComponent)
 {
-	if (GSimModuleDebugParams.EnableNetworkStateData && NetworkComponent)
+	if (GSimModuleDebugParams.EnableNetworkStateData)
 	{
-		if (const FModularVehicleSimulationCU* VehicleSimulation = Cast<const UModularVehicleBaseComponent>(NetworkComponent)->VehicleSimulationPT.Get())
+		if (const UModularVehicleBaseComponent* ModularBaseComponent = Cast<const UModularVehicleBaseComponent>(NetworkComponent))
 		{
-			VehicleInputs = VehicleSimulation->VehicleInputs;
+			if (const FModularVehicleSimulationCU* VehicleSimulation = ModularBaseComponent->VehicleSimulationPT.Get())
+			{
+				VehicleInputs = VehicleSimulation->VehicleInputs;
+			}
 		}
 	}
 }
@@ -303,34 +309,37 @@ bool FNetworkModularVehicleStates::NetSerialize(FArchive& Ar, class UPackageMap*
 
 	int32 NumNetModules = ModuleData.Num();
 	Ar << NumNetModules;
-
+	if(Ar.IsLoading() && NumNetModules != ModuleData.Num())
+	{
+		ModuleData.Reserve(NumNetModules);
+	}
 	for (int I = 0; I < NumNetModules; I++)
 	{
 		if (Ar.IsLoading())
 		{
 			if (NumNetModules > 0)
 			{
-				int32 ModuleType = Chaos::eSimType::Undefined;
+				FName ModuleType = NAME_None;
 				int32 SimArrayIndex = 0;
 				Ar << ModuleType;
 				Ar << SimArrayIndex;
 
-				if (ModuleData.Num() != NumNetModules)
+				if (I >= ModuleData.Num())
 				{
-					ModuleData.Reserve(NumNetModules);
-
 					if (TSharedPtr<Chaos::FModuleNetData> Data = Chaos::FModuleFactoryRegister::Get().GenerateNetData(ModuleType, SimArrayIndex))
 					{
-						Data->Serialize(Ar);
 						ModuleData.Emplace(Data);
 					}
-
+				}
+				if(I <= ModuleData.Num() && ModuleData[I].IsValid())
+				{
+					ModuleData[I]->Serialize(Ar);
 				}
 			}
 		}
 		else
 		{
-			int32 ModuleType = (int32)ModuleData[I]->GetType();
+			FName ModuleType = ModuleData[I]->GetSimType();
 			Ar << ModuleType;
 			Ar << ModuleData[I]->SimArrayIndex;
 			ModuleData[I]->Serialize(Ar);
@@ -343,9 +352,12 @@ bool FNetworkModularVehicleStates::NetSerialize(FArchive& Ar, class UPackageMap*
 
 void FNetworkModularVehicleStates::ApplyData(UActorComponent* NetworkComponent) const
 {
-	if (FModularVehicleSimulationCU* VehicleSimulation = Cast<UModularVehicleBaseComponent>(NetworkComponent)->VehicleSimulationPT.Get())
+	if (UModularVehicleBaseComponent* ModularBaseComponent = Cast<UModularVehicleBaseComponent>(NetworkComponent))
 	{
-		VehicleSimulation->AccessSimComponentTree()->SetSimState(ModuleData);
+		if (FModularVehicleSimulationCU* VehicleSimulation = ModularBaseComponent->VehicleSimulationPT.Get())
+		{
+			VehicleSimulation->AccessSimComponentTree()->SetSimState(ModuleData);
+		}
 	}
 }
 
@@ -370,8 +382,8 @@ void FNetworkModularVehicleStates::InterpolateData(const FNetworkPhysicsData& Mi
 	for (int I = 0; I < ModuleData.Num(); I++)
 	{
 		// if these don't match then something has gone terribly wrong
-		check(ModuleData[I]->GetType() == MinState.ModuleData[I]->GetType());
-		check(ModuleData[I]->GetType() == MaxState.ModuleData[I]->GetType());
+		check(ModuleData[I]->GetSimType() == MinState.ModuleData[I]->GetSimType());
+		check(ModuleData[I]->GetSimType() == MaxState.ModuleData[I]->GetSimType());
 
 		ModuleData[I]->Lerp(LerpFactor, *MinState.ModuleData[I].Get(), *MaxState.ModuleData[I].Get());
 	}
diff --git a/Engine/Plugins/Experimental/ChaosModularVehicle/Source/ChaosModularVehicleEngine/Private/ChaosModularVehicle/ModularVehicleBaseComponent.cpp b/Engine/Plugins/Experimental/ChaosModularVehicle/Source/ChaosModularVehicleEngine/Private/ChaosModularVehicle/ModularVehicleBaseComponent.cpp
--- a/Engine/Plugins/Experimental/ChaosModularVehicle/Source/ChaosModularVehicleEngine/Private/ChaosModularVehicle/ModularVehicleBaseComponent.cpp
+++ b/Engine/Plugins/Experimental/ChaosModularVehicle/Source/ChaosModularVehicleEngine/Private/ChaosModularVehicle/ModularVehicleBaseComponent.cpp
@@ -658,12 +658,12 @@ void UModularVehicleBaseComponent::ParallelUpdate(float DeltaTime)
 					}
 
 					// extract/cache some generally useful values as we go as trying to locate this data later requires a search
-					if (PVehicleOutput->SimTreeOutputData[I]->GetType() == Chaos::eSimType::Transmission)
+					if (PVehicleOutput->SimTreeOutputData[I]->IsSimType<Chaos::FTransmissionSimModule>())
 					{
 						// if there is more than one transmission then the last one will inform us of the current gear
 						CurrentGear = static_cast<Chaos::FTransmissionOutputData*>(PVehicleOutput->SimTreeOutputData[I])->CurrentGear;
 					}
-					else if (PVehicleOutput->SimTreeOutputData[I]->GetType() == Chaos::eSimType::Engine)
+					else if (PVehicleOutput->SimTreeOutputData[I]->IsSimType<Chaos::FEngineSimModule>())
 					{
 						// if there is more than one engine then the last one will inform us of the engine RPM
 						Chaos::FEngineOutputData* Engine = static_cast<Chaos::FEngineOutputData*>(CurrentOutput->VehicleSimOutput.SimTreeOutputData[I]);
@@ -906,11 +906,12 @@ void UModularVehicleBaseComponent::RemoveComponentFromSimulation(UPrimitiveCompo
 		Chaos::FPBDRigidsSolver* Solver = Proxy->GetSolver<Chaos::FPBDRigidsSolver>();
 		Solver->EnqueueCommandImmediate([Proxy, this, LatestTreeUpdates = LatestTreeUpdates]() mutable
 			{
-
 				if (VehicleSimulationPT)
 				{
-					TUniquePtr<Chaos::FSimModuleTree>& SimModuleTree = VehicleSimulationPT->AccessSimComponentTree();
-					SimModuleTree->AppendTreeUpdates(LatestTreeUpdates);
+					if(TUniquePtr<Chaos::FSimModuleTree>& SimModuleTree = VehicleSimulationPT->AccessSimComponentTree())
+					{
+						SimModuleTree->AppendTreeUpdates(LatestTreeUpdates);
+					}
 				}
 			});
 
diff --git a/Engine/Plugins/Experimental/ChaosModularVehicle/Source/ChaosModularVehicleEngine/Private/ChaosModularVehicle/ModularVehicleBuilder.cpp b/Engine/Plugins/Experimental/ChaosModularVehicle/Source/ChaosModularVehicleEngine/Private/ChaosModularVehicle/ModularVehicleBuilder.cpp
--- a/Engine/Plugins/Experimental/ChaosModularVehicle/Source/ChaosModularVehicleEngine/Private/ChaosModularVehicle/ModularVehicleBuilder.cpp
+++ b/Engine/Plugins/Experimental/ChaosModularVehicle/Source/ChaosModularVehicleEngine/Private/ChaosModularVehicle/ModularVehicleBuilder.cpp
@@ -35,13 +35,13 @@ void FModularVehicleBuilder::FixupTreeLinks(TUniquePtr<Chaos::FSimModuleTree>& S
 {
 	using namespace Chaos;
 
-	const FSimModuleTree::FSimModuleNode* TransmissionNode = SimModuleTree->LocateNodeByType(eSimType::Transmission);
+	const FSimModuleTree::FSimModuleNode* TransmissionNode = SimModuleTree->LocateNodeByType<Chaos::FTransmissionSimModule>();
 
 	for (int I = 0; I < SimModuleTree->NumActiveNodes(); I++)
 	{
 		// link suspension and wheel references to one another
 		ISimulationModuleBase* Module = SimModuleTree->AccessSimModule(I);
-		if (Module && Module->GetSimType() == eSimType::Suspension)
+		if (Module && Module->IsSimType<FSuspensionBaseInterface>())
 		{
 			FSimModuleTree::FSimModuleNode& SuspensionNode = SimModuleTree->GetNode(I);
 
@@ -51,12 +51,13 @@ void FModularVehicleBuilder::FixupTreeLinks(TUniquePtr<Chaos::FSimModuleTree>& S
 
 				ISimulationModuleBase* ParentModule = SimModuleTree->AccessSimModule(SuspensionParentNode.SimModule->GetTreeIndex());
 
-				if (ParentModule && ParentModule->GetSimType() == eSimType::Wheel)
+				if(FWheelBaseInterface* Wheel = ParentModule->Cast<FWheelBaseInterface>())
 				{
-					FSuspensionBaseInterface* Suspension = static_cast<FSuspensionBaseInterface*>(Module);
-					FWheelBaseInterface* Wheel = static_cast<FWheelBaseInterface*>(ParentModule);
-					Wheel->SetSuspensionSimTreeIndex(Suspension->GetTreeIndex());
-					Suspension->SetWheelSimTreeIndex(Wheel->GetTreeIndex());
+					if (FSuspensionBaseInterface* Suspension = Module->Cast<FSuspensionBaseInterface>())
+					{
+						Wheel->SetSuspensionSimTreeIndex(Suspension->GetTreeIndex());
+						Suspension->SetWheelSimTreeIndex(Wheel->GetTreeIndex());
+					}
 				}
 			}
 
@@ -68,12 +69,13 @@ void FModularVehicleBuilder::FixupTreeLinks(TUniquePtr<Chaos::FSimModuleTree>& S
 
 					ISimulationModuleBase* ChildModule = SimModuleTree->AccessSimModule(SuspensionChildNode.SimModule->GetTreeIndex());
 
-					if (ChildModule && ChildModule->GetSimType() == eSimType::Wheel)
+					if(FWheelBaseInterface* Wheel = ChildModule->Cast<FWheelBaseInterface>())
 					{
-						FSuspensionBaseInterface* Suspension = static_cast<FSuspensionBaseInterface*>(Module);
-						FWheelBaseInterface* Wheel = static_cast<FWheelBaseInterface*>(ChildModule);
-						Wheel->SetSuspensionSimTreeIndex(Suspension->GetTreeIndex());
-						Suspension->SetWheelSimTreeIndex(Wheel->GetTreeIndex());
+						if (FSuspensionBaseInterface* Suspension = Module->Cast<FSuspensionBaseInterface>())
+						{
+							Wheel->SetSuspensionSimTreeIndex(Suspension->GetTreeIndex());
+							Suspension->SetWheelSimTreeIndex(Wheel->GetTreeIndex());
+						}
 					}
 				}
 			}
@@ -82,7 +84,7 @@ void FModularVehicleBuilder::FixupTreeLinks(TUniquePtr<Chaos::FSimModuleTree>& S
 		// temporarily - link suspension/wheels to engine/clutch/transmission, assuming only 1 transmission & all wheels are powered
 		if (TransmissionNode)
 		{
-			if (Module && Module->GetSimType() == eSimType::Wheel)
+			if (Module && Module->IsSimType<class FWheelSimModule>())
 			{
 				SimModuleTree->Reparent(Module->GetTreeIndex(), TransmissionNode->SimModule->GetTreeIndex());
 			}
diff --git a/Engine/Plugins/Experimental/ChaosModularVehicle/Source/ChaosModularVehicleEngine/Private/ChaosModularVehicle/SuspensionSimModule.cpp b/Engine/Plugins/Experimental/ChaosModularVehicle/Source/ChaosModularVehicleEngine/Private/ChaosModularVehicle/SuspensionSimModule.cpp
--- a/Engine/Plugins/Experimental/ChaosModularVehicle/Source/ChaosModularVehicleEngine/Private/ChaosModularVehicle/SuspensionSimModule.cpp
+++ b/Engine/Plugins/Experimental/ChaosModularVehicle/Source/ChaosModularVehicleEngine/Private/ChaosModularVehicle/SuspensionSimModule.cpp
@@ -25,12 +25,6 @@ using namespace Chaos;
 		, SpringSpeed(0.f)
 	{
 		AccessSetup().MaxLength = FMath::Abs(Settings.MaxRaise + Settings.MaxDrop);
-
-		if (!FModuleFactoryRegister::Get().ContainsFactory(GetSimType()))
-		{
-			static TSharedPtr<FSuspensionFactory> SharedFactory = MakeShared<FSuspensionFactory>();
-			FModuleFactoryRegister::Get().RegisterFactory(GetSimType(), SharedFactory);
-		}
 	}
 	FSuspensionSimModule::~FSuspensionSimModule()
 	{
@@ -102,10 +96,10 @@ using namespace Chaos;
 			{
 				if (Chaos::ISimulationModuleBase* Module = SimModuleTree->AccessSimModule(WheelSimTreeIndex))
 				{
-					check(Module->GetSimType() == eSimType::Wheel);
-					Chaos::FWheelBaseInterface* Wheel = static_cast<Chaos::FWheelBaseInterface*>(Module);
-
-					Wheel->SetForceIntoSurface(ForceIntoSurface);
+					if(Chaos::FWheelBaseInterface* Wheel = Module->Cast<Chaos::FWheelBaseInterface>())
+					{
+						Wheel->SetForceIntoSurface(ForceIntoSurface);
+					}
 				}
 		
 			}
@@ -193,8 +187,7 @@ using namespace Chaos;
 
 	void FSuspensionSimModuleData::FillSimState(ISimulationModuleBase* SimModule)
 	{
-		check(SimModule->GetSimType() == eSimType::Suspension);
-		if (FSuspensionSimModule* Sim = static_cast<FSuspensionSimModule*>(SimModule))
+		if (FSuspensionSimModule* Sim = SimModule->Cast<FSuspensionSimModule>())
 		{
 			Sim->SpringDisplacement = SpringDisplacement;
 			Sim->LastDisplacement = LastDisplacement;
@@ -203,8 +196,7 @@ using namespace Chaos;
 
 	void FSuspensionSimModuleData::FillNetState(const ISimulationModuleBase* SimModule)
 	{
-		check(SimModule->GetSimType() == eSimType::Suspension);
-		if (const FSuspensionSimModule* Sim = static_cast<const FSuspensionSimModule*>(SimModule))
+		if (const FSuspensionSimModule* Sim = SimModule->Cast<const FSuspensionSimModule>())
 		{
 			SpringDisplacement = Sim->SpringDisplacement;
 			LastDisplacement = Sim->LastDisplacement;
@@ -230,11 +222,9 @@ using namespace Chaos;
 
 	void FSuspensionOutputData::FillOutputState(const ISimulationModuleBase* SimModule)
 	{
-		check(SimModule->GetSimType() == eSimType::Suspension);
-
 		FSimOutputData::FillOutputState(SimModule);
 
-		if (const FSuspensionSimModule* Sim = static_cast<const FSuspensionSimModule*>(SimModule))
+		if (const FSuspensionSimModule* Sim = SimModule->Cast<const FSuspensionSimModule>())
 		{
 			SpringDisplacement = Sim->SpringDisplacement;
 			SpringDisplacementVector = -Sim->Setup().SuspensionAxis * Sim->SpringDisplacement + Sim->GetAnimationOffset();
diff --git a/Engine/Plugins/Experimental/ChaosModularVehicle/Source/ChaosModularVehicleEngine/Public/ChaosModularVehicle/SuspensionSimModule.h b/Engine/Plugins/Experimental/ChaosModularVehicle/Source/ChaosModularVehicleEngine/Public/ChaosModularVehicle/SuspensionSimModule.h
--- a/Engine/Plugins/Experimental/ChaosModularVehicle/Source/ChaosModularVehicleEngine/Public/ChaosModularVehicle/SuspensionSimModule.h
+++ b/Engine/Plugins/Experimental/ChaosModularVehicle/Source/ChaosModularVehicleEngine/Public/ChaosModularVehicle/SuspensionSimModule.h
@@ -5,6 +5,8 @@
 #include "SimModule/SuspensionBaseInterface.h"
 #include "Chaos/ChaosEngineInterface.h"
 
+class FSuspensionSimModule;
+
 namespace Chaos
 {
 	struct FAllInputs;
@@ -12,152 +14,167 @@ namespace Chaos
 	class FClusterUnionPhysicsProxy;
 }
 
-	struct CHAOSMODULARVEHICLEENGINE_API FSuspensionSimModuleData : public Chaos::FModuleNetData
-	{
+struct CHAOSMODULARVEHICLEENGINE_API FSuspensionSimModuleData 
+	: public Chaos::FModuleNetData
+	, public Chaos::TSimulationModuleTypeable<FSuspensionSimModule,FSuspensionSimModuleData>
+{
 #if !(UE_BUILD_SHIPPING || UE_BUILD_TEST)
-		FSuspensionSimModuleData(int NodeArrayIndex, const FString& InDebugString) : Chaos::FModuleNetData(NodeArrayIndex, InDebugString) {}
+	FSuspensionSimModuleData(int NodeArrayIndex, const FString& InDebugString) : Chaos::FModuleNetData(NodeArrayIndex, InDebugString) {}
 #else
-		FSuspensionSimModuleData(int NodeArrayIndex) : FModuleNetData(NodeArrayIndex) {}
+	FSuspensionSimModuleData(int NodeArrayIndex) : FModuleNetData(NodeArrayIndex) {}
 #endif
 
-		virtual Chaos::eSimType GetType() override { return Chaos::eSimType::Suspension; }
-
-		virtual void FillSimState(Chaos::ISimulationModuleBase* SimModule) override;
+	virtual void FillSimState(Chaos::ISimulationModuleBase* SimModule) override;
 
-		virtual void FillNetState(const Chaos::ISimulationModuleBase* SimModule) override;
+	virtual void FillNetState(const Chaos::ISimulationModuleBase* SimModule) override;
 
-		virtual void Serialize(FArchive& Ar) override
-		{
-			Ar << SpringDisplacement;
-			Ar << LastDisplacement;
-		}
+	virtual void Serialize(FArchive& Ar) override
+	{
+		Ar << SpringDisplacement;
+		Ar << LastDisplacement;
+	}
 
-		virtual void Lerp(const float LerpFactor, const FModuleNetData& Min, const FModuleNetData& Max) override;
+	virtual void Lerp(const float LerpFactor, const FModuleNetData& Min, const FModuleNetData& Max) override;
 
 #if !(UE_BUILD_SHIPPING || UE_BUILD_TEST)
-		virtual FString ToString() const override;
+	virtual FString ToString() const override;
 #endif
 
-		float SpringDisplacement = 0.0f;
-		float LastDisplacement = 0.0f;
-	};
+	float SpringDisplacement = 0.0f;
+	float LastDisplacement = 0.0f;
+};
 
-	struct CHAOSMODULARVEHICLEENGINE_API FSuspensionOutputData : public Chaos::FSimOutputData
-	{
-		virtual FSimOutputData* MakeNewData() override { return FSuspensionOutputData::MakeNew(); }
-		static FSimOutputData* MakeNew() { return new FSuspensionOutputData(); }
+struct CHAOSMODULARVEHICLEENGINE_API FSuspensionOutputData 
+	: public Chaos::FSimOutputData
+	, public Chaos::TSimulationModuleTypeable<FSuspensionSimModule,FSuspensionOutputData>
+{
+	virtual FSimOutputData* MakeNewData() override { return FSuspensionOutputData::MakeNew(); }
+	static FSimOutputData* MakeNew() { return new FSuspensionOutputData(); }
 
-		virtual Chaos::eSimType GetType() override { return Chaos::eSimType::Suspension; }
-		virtual void FillOutputState(const Chaos::ISimulationModuleBase* SimModule) override;
-		virtual void Lerp(const FSimOutputData& InCurrent, const Chaos::FSimOutputData& InNext, float Alpha) override;
+	virtual void FillOutputState(const Chaos::ISimulationModuleBase* SimModule) override;
+	virtual void Lerp(const FSimOutputData& InCurrent, const Chaos::FSimOutputData& InNext, float Alpha) override;
 
 #if !(UE_BUILD_SHIPPING || UE_BUILD_TEST)
-		virtual FString ToString() override;
+	virtual FString ToString() override;
 #endif
 
-		float SpringDisplacement;
-		FVector SpringDisplacementVector;
-		float SpringSpeed;
-		FVector ImpactNormal;
-	};
+	float SpringDisplacement;
+	FVector SpringDisplacementVector;
+	float SpringSpeed;
+	FVector ImpactNormal;
+};
 
-	struct CHAOSMODULARVEHICLEENGINE_API FSuspensionSettings
+struct CHAOSMODULARVEHICLEENGINE_API FSuspensionSettings
+{
+	FSuspensionSettings()
+		: SuspensionAxis(FVector(0.f, 0.f, -1.f))
+		, RestOffset(FVector::ZeroVector)
+		, MaxRaise(5.f)
+		, MaxDrop(5.f)
+		, MaxLength(0.f)
+		, SpringRate(1.f)
+		, SpringPreload(0.5f)
+		, SpringDamping(0.9f)
+		, SuspensionForceEffect(100.0f)
 	{
-		FSuspensionSettings()
-			: SuspensionAxis(FVector(0.f, 0.f, -1.f))
-			, RestOffset(FVector::ZeroVector)
-			, MaxRaise(5.f)
-			, MaxDrop(5.f)
-			, MaxLength(0.f)
-			, SpringRate(1.f)
-			, SpringPreload(0.5f)
-			, SpringDamping(0.9f)
-			, SuspensionForceEffect(100.0f)
-		{
-
-		}
-
-		FVector SuspensionAxis;		// local axis, direction of suspension force raycast traces
-		FVector RestOffset;
-		float MaxRaise;				// distance [cm]
-		float MaxDrop;				// distance [cm]
-		float MaxLength;			// distance [cm]
-
-		float SpringRate;			// spring constant
-		float SpringPreload;		// Amount of Spring force (independent spring movement)
-		float SpringDamping;		// limit compression/rebound speed
-
-		float SuspensionForceEffect; // force that presses the wheels into the ground - producing grip
-	};
-	
 
-	class FSuspensionFactory : public Chaos::IFactoryModule
+	}
+
+	FVector SuspensionAxis;		// local axis, direction of suspension force raycast traces
+	FVector RestOffset;
+	float MaxRaise;				// distance [cm]
+	float MaxDrop;				// distance [cm]
+	float MaxLength;			// distance [cm]
+
+	float SpringRate;			// spring constant
+	float SpringPreload;		// Amount of Spring force (independent spring movement)
+	float SpringDamping;		// limit compression/rebound speed
+
+	float SuspensionForceEffect; // force that presses the wheels into the ground - producing grip
+};
+
+
+class FSuspensionFactory : public Chaos::IFactoryModule
+{
+public:
+	virtual TSharedPtr<Chaos::FModuleNetData> GenerateNetData(const int32 SimArrayIndex) const override
 	{
-	public:
-		TSharedPtr<Chaos::FModuleNetData> GenerateNetData(int32 SimArrayIndex) const
-		{
-			return MakeShared<FSuspensionSimModuleData>(
-				SimArrayIndex
+		return MakeShared<FSuspensionSimModuleData>(
+			SimArrayIndex
 #if !(UE_BUILD_SHIPPING || UE_BUILD_TEST)
-				, TEXT("ConstraintSuspension")
+			, TEXT("ConstraintSuspension")
 #endif			
-			);
-		}
-	};
-
-	class CHAOSMODULARVEHICLEENGINE_API FSuspensionSimModule : public Chaos::FSuspensionBaseInterface, public Chaos::TSimModuleSettings<FSuspensionSettings>
-	{
-		friend FSuspensionSimModuleData;
-		friend FSuspensionOutputData;
-
-	public:
+		);
+	}
+};
+
+class CHAOSMODULARVEHICLEENGINE_API FSuspensionSimModule 
+	: public Chaos::FSuspensionBaseInterface
+	, public Chaos::TSimModuleSettings<FSuspensionSettings>
+	, public Chaos::TSimulationModuleTypeable<FSuspensionSimModule>
+{
+	friend FSuspensionSimModuleData;
+	friend FSuspensionOutputData;
 
-		FSuspensionSimModule(const FSuspensionSettings& Settings);
+public:
+	DEFINE_CHAOSSIMTYPENAME(FSuspensionSimModule);
+	FSuspensionSimModule(const FSuspensionSettings& Settings);
 
-		virtual TSharedPtr<Chaos::FModuleNetData> GenerateNetData(int SimArrayIndex) const override
-		{
-			return MakeShared<FSuspensionSimModuleData>(
-				SimArrayIndex
+	virtual TSharedPtr<Chaos::FModuleNetData> GenerateNetData(const int32 SimArrayIndex) const override
+	{
+		return MakeShared<FSuspensionSimModuleData>(
+			SimArrayIndex
 #if !(UE_BUILD_SHIPPING || UE_BUILD_TEST)
-				, GetDebugName()
+			, GetDebugName()
 #endif			
-			);
-		}
-		virtual ~FSuspensionSimModule() override;
+		);
+	}
+	virtual ~FSuspensionSimModule() override;
+
+
+	virtual Chaos::FSimOutputData* GenerateOutputData() const override
+	{
+		return FSuspensionOutputData::MakeNew();
+	}
 
+	virtual const FString GetDebugName() const { return TEXT("Suspension"); }
 
-		virtual Chaos::FSimOutputData* GenerateOutputData() const override
-		{
-			return FSuspensionOutputData::MakeNew();
-		}
+	virtual float GetMaxSpringLength() const override { return Setup().MaxLength; }
+	virtual float GetSpringLength() const override;
+	virtual  void SetSpringLength(float InLength, float WheelRadius) override;
+	virtual void GetWorldRaycastLocation(const FTransform& BodyTransform, float WheelRadius, Chaos::FSpringTrace& OutTrace) override;
 
-		virtual const FString GetDebugName() const { return TEXT("Suspension"); }
+	virtual void OnConstruction_External(Chaos::FClusterUnionPhysicsProxy* Proxy) override;
+	virtual void OnTermination_External() override;
 
-		virtual float GetMaxSpringLength() const override { return Setup().MaxLength; }
-		virtual float GetSpringLength() const override;
-		virtual  void SetSpringLength(float InLength, float WheelRadius) override;
-		virtual void GetWorldRaycastLocation(const FTransform& BodyTransform, float WheelRadius, Chaos::FSpringTrace& OutTrace) override;
+	virtual void Simulate(float DeltaTime, const Chaos::FAllInputs& Inputs, Chaos::FSimModuleTree& VehicleModuleSystem) override;
 
-		virtual void OnConstruction_External(Chaos::FClusterUnionPhysicsProxy* Proxy) override;
-		virtual void OnTermination_External() override;
+	virtual void Animate(Chaos::FClusterUnionPhysicsProxy* Proxy) override;
 
-		virtual void Simulate(float DeltaTime, const Chaos::FAllInputs& Inputs, Chaos::FSimModuleTree& VehicleModuleSystem) override;
+	const FVector& GetRestLocation() const { return Setup().RestOffset; }
 
-		virtual void Animate(Chaos::FClusterUnionPhysicsProxy* Proxy) override;
+	void UpdateConstraint();
 
-		const FVector& GetRestLocation() const { return Setup().RestOffset; }
+protected:
+	void CreateConstraint(Chaos::FClusterUnionPhysicsProxy* Proxy);
+	void DestroyConstraint();
 
-		void UpdateConstraint();
+private:
 
-	protected:
-		void CreateConstraint(Chaos::FClusterUnionPhysicsProxy* Proxy);
-		void DestroyConstraint();
+	float SpringDisplacement;
+	float LastDisplacement;
+	float SpringSpeed;
 
-	private:
+	FPhysicsConstraintHandle ConstraintHandle;
+};
 
-		float SpringDisplacement;
-		float LastDisplacement;
-		float SpringSpeed;
 
-		FPhysicsConstraintHandle ConstraintHandle;
-	};
+class CHAOSMODULARVEHICLEENGINE_API FSuspensionSimFactory
+		: public Chaos::FSimFactoryModule<FSuspensionSimModuleData>
+		, public Chaos::TSimulationModuleTypeable<FSuspensionSimModule, FSuspensionSimFactory>
+		, public Chaos::TSimFactoryAutoRegister<FSuspensionSimFactory>
+	
+{
+public:
+	FSuspensionSimFactory() : FSimFactoryModule(TEXT("SuspensionSimFactory")) {}
+};
\ No newline at end of file
diff --git a/Engine/Source/Runtime/Engine/Private/PhysicsEngine/NetworkPhysicsComponent.cpp b/Engine/Source/Runtime/Engine/Private/PhysicsEngine/NetworkPhysicsComponent.cpp
--- a/Engine/Source/Runtime/Engine/Private/PhysicsEngine/NetworkPhysicsComponent.cpp
+++ b/Engine/Source/Runtime/Engine/Private/PhysicsEngine/NetworkPhysicsComponent.cpp
@@ -460,6 +460,7 @@ void UNetworkPhysicsComponent::InitPhysics()
 	* If the NetworkPhysicsComponent is added as a SubObject after the actor has processed bAutoActivate and bWantsInitializeComponent
 	* SetActive(true) and InitializeComponent() needs to be called manually for the component to function properly. */
 	PrimaryComponentTick.bCanEverTick = true;
+	PrimaryComponentTick.TickGroup = TG_PostPhysics;
 	bAutoActivate = true;
 	bWantsInitializeComponent = true;
 	SetIsReplicatedByDefault(true);
@@ -525,6 +526,15 @@ void UNetworkPhysicsComponent::InitializeComponent()
 void UNetworkPhysicsComponent::UninitializeComponent()
 {
 	Super::UninitializeComponent();
+
+	if (NetworkPhysicsComponent_Internal)
+	{
+		if (FAsyncNetworkPhysicsComponentInput* AsyncInput = NetworkPhysicsComponent_Internal->GetProducerInputData_External())
+		{
+			AsyncInput->ActorComponent = nullptr;
+		}
+	}
+
 	if (UWorld* World = GetWorld())
 	{
 		if (FPhysScene* PhysScene = World->GetPhysicsScene())
@@ -1178,9 +1188,9 @@ void FAsyncNetworkPhysicsComponent::ConsumeAsyncInput()
 			if (AsyncInput->InputData && AsyncInput->InputData->HasDataInHistory())
 			{
 				// Validate data in the received inputs on the server
-				if (!ComponentSettings.GetValidateDataOnGameThread() && IsServer() && ActorComponent)
+				if (!ComponentSettings.GetValidateDataOnGameThread() && IsServer() && ActorComponent.IsValid() && ActorComponent.Get()->IsBeingDestroyed() == false)
 				{
-					AsyncInput->InputData->ValidateDataInHistory(ActorComponent);
+					AsyncInput->InputData->ValidateDataInHistory(ActorComponent.Get());
 				}
 				ReceiveHelper(InputHistory.Get(), AsyncInput->InputData.Get(), /*bImportant*/false, ComponentSettings.GetCompareInputToTriggerRewind());
 			}
@@ -1259,7 +1269,7 @@ void FAsyncNetworkPhysicsComponent::OnPreProcessInputs_Internal(const int32 Phys
 	}
 #endif
 
-	if (ActorComponent)
+	if (ActorComponent.IsValid() && ActorComponent.Get()->IsBeingDestroyed() == false)
 	{
 		// Apply replicated state on clients if we are resimulating
 		if (bIsSolverResim && StateHistory && StateData)
@@ -1269,7 +1279,7 @@ void FAsyncNetworkPhysicsComponent::OnPreProcessInputs_Internal(const int32 Phys
 			const bool bExactFrame = PhysicsReplicationCVars::ResimulationCVars::bAllowRewindToClosestState ? !bIsSolverReset : true;
 			if (StateHistory->ExtractData(PhysicsStep, bIsSolverReset, PhysicsData, bExactFrame) && PhysicsData->bReceivedData)
 			{
-				PhysicsData->ApplyData(ActorComponent);
+				PhysicsData->ApplyData(ActorComponent.Get());
 #if DEBUG_NETWORK_PHYSICS
 				UE_LOG(LogChaos, Log, TEXT("		Applying extracted state from history | bExactFrame = %d | LocalFrame = %d | ServerFrame = %d | InputFrame = %d | Data: %s")
 					, bExactFrame, PhysicsData->LocalFrame, PhysicsData->ServerFrame, PhysicsData->InputFrame, *PhysicsData->DebugData());
@@ -1300,7 +1310,7 @@ void FAsyncNetworkPhysicsComponent::OnPreProcessInputs_Internal(const int32 Phys
 					UE_LOG(LogChaos, Log, TEXT("	Reapplying multiple data due to receiving an important data that was previously missed. FromFrame: %d | ToFrame: %d | IsLocallyControlled = %d"), NewImportantInputFrame, (NextExpectedLocalFrame - 1), IsLocallyControlled());
 #endif
 					// Apply all inputs in range
-					InputHistory->ApplyDataRange(NewImportantInputFrame, NextExpectedLocalFrame - 1, ActorComponent, /*bOnlyImportant*/false);
+					InputHistory->ApplyDataRange(NewImportantInputFrame, NextExpectedLocalFrame - 1, ActorComponent.Get(), /*bOnlyImportant*/false);
 				}
 				else
 				{
@@ -1337,7 +1347,7 @@ void FAsyncNetworkPhysicsComponent::OnPreProcessInputs_Internal(const int32 Phys
 						{
 							if (InputHistory->ExtractData(NextExpectedLocalFrame, bIsSolverReset, PhysicsData, true) && NextExpectedLocalFrame < LastFrame)
 							{
-								PhysicsData->ApplyData(ActorComponent);
+								PhysicsData->ApplyData(ActorComponent.Get());
 							}
 						}
 					}
@@ -1366,7 +1376,7 @@ void FAsyncNetworkPhysicsComponent::OnPreProcessInputs_Internal(const int32 Phys
 					InputHistory->RecordData(PhysicsStep, PhysicsData);
 				}
 
-				PhysicsData->ApplyData(ActorComponent);
+				PhysicsData->ApplyData(ActorComponent.Get());
 
 #if DEBUG_NETWORK_PHYSICS
 				{
@@ -1407,7 +1417,7 @@ void FAsyncNetworkPhysicsComponent::OnPostProcessInputs_Internal(const int32 Phy
 	}
 #endif
 
-	if (ActorComponent)
+	if (ActorComponent.IsValid() && ActorComponent.Get()->IsBeingDestroyed() == false)
 	{
 		const bool bShouldCacheInputHistory = IsLocallyControlled() && !bIsSolverResim;
 		// For the inputs client local ones are ground truth otherwise use the replicated ones coming from the server
@@ -1418,7 +1428,7 @@ void FAsyncNetworkPhysicsComponent::OnPostProcessInputs_Internal(const int32 Phy
 			PhysicsData->PrepareFrame(PhysicsStep, bIsServer, GetNetworkPhysicsTickOffset());
 
 			// Gather input data from implementation
-			PhysicsData->BuildData(ActorComponent);
+			PhysicsData->BuildData(ActorComponent.Get());
 
 			// Record input in history
 			InputHistory->RecordData(PhysicsStep, PhysicsData);
@@ -1451,7 +1461,7 @@ void FAsyncNetworkPhysicsComponent::OnPostProcessInputs_Internal(const int32 Phy
 			PhysicsData->InputFrame = InputFrame;
 
 			// Gather input data from implementation
-			PhysicsData->BuildData(ActorComponent);
+			PhysicsData->BuildData(ActorComponent.Get());
 
 			// Record input in history
 			StateHistory->RecordData(PhysicsStep, PhysicsData);
diff --git a/Engine/Source/Runtime/Engine/Public/Physics/NetworkPhysicsComponent.h b/Engine/Source/Runtime/Engine/Public/Physics/NetworkPhysicsComponent.h
--- a/Engine/Source/Runtime/Engine/Public/Physics/NetworkPhysicsComponent.h
+++ b/Engine/Source/Runtime/Engine/Public/Physics/NetworkPhysicsComponent.h
@@ -1189,7 +1189,7 @@ private:
 	static const FNetworkPhysicsSettingsNetworkPhysicsComponent SettingsNetworkPhysicsComponent_Default;
 
 	// Actor component that will be used to fill the histories
-	TObjectPtr<UActorComponent> ActorComponent;
+	TWeakObjectPtr<UActorComponent> ActorComponent;
 
 	// Root components physics object
 	Chaos::FPhysicsObjectHandle RootPhysicsObject;
diff --git a/Engine/Source/Runtime/Experimental/ChaosVehicles/ChaosVehiclesCore/Private/SimModule/ClutchModule.cpp b/Engine/Source/Runtime/Experimental/ChaosVehicles/ChaosVehiclesCore/Private/SimModule/ClutchModule.cpp
--- a/Engine/Source/Runtime/Experimental/ChaosVehicles/ChaosVehiclesCore/Private/SimModule/ClutchModule.cpp
+++ b/Engine/Source/Runtime/Experimental/ChaosVehicles/ChaosVehiclesCore/Private/SimModule/ClutchModule.cpp
@@ -13,11 +13,6 @@ namespace Chaos
 	FClutchSimModule::FClutchSimModule(const FClutchSettings& Settings) : TSimModuleSettings<FClutchSettings>(Settings)
 		, ClutchValue(0.f)
 	{
-		if (!FModuleFactoryRegister::Get().ContainsFactory(GetSimType()))
-		{
-			static TSharedPtr<FSimFactoryModule<FClutchSimModuleData>> SharedFactory = MakeShared<FSimFactoryModule<FClutchSimModuleData>>(GetDebugName());
-			FModuleFactoryRegister::Get().RegisterFactory(GetSimType(), SharedFactory);
-		}
 	}
 
 	void FClutchSimModule::Simulate(float DeltaTime, const FAllInputs& Inputs, FSimModuleTree& VehicleModuleSystem)
diff --git a/Engine/Source/Runtime/Experimental/ChaosVehicles/ChaosVehiclesCore/Private/SimModule/EngineModule.cpp b/Engine/Source/Runtime/Experimental/ChaosVehicles/ChaosVehiclesCore/Private/SimModule/EngineModule.cpp
--- a/Engine/Source/Runtime/Experimental/ChaosVehicles/ChaosVehiclesCore/Private/SimModule/EngineModule.cpp
+++ b/Engine/Source/Runtime/Experimental/ChaosVehicles/ChaosVehiclesCore/Private/SimModule/EngineModule.cpp
@@ -15,12 +15,6 @@ namespace Chaos
 		, MaxEngineSpeed(RPMToOmega(Setup().MaxRPM))
 		, EngineStarted(true)
 	{
-		if (!FModuleFactoryRegister::Get().ContainsFactory(GetSimType()))
-		{
-			static TSharedPtr<FSimFactoryModule<FEngineSimModuleData>> SharedFactory = MakeShared<FSimFactoryModule<FEngineSimModuleData>>(GetDebugName());
-			FModuleFactoryRegister::Get().RegisterFactory(GetSimType(), SharedFactory);
-		}
-
 	}
 
 	void FEngineSimModule::Simulate(float DeltaTime, const FAllInputs& Inputs, FSimModuleTree& VehicleModuleSystem)
@@ -90,7 +84,7 @@ namespace Chaos
 
 	void FEngineOutputData::FillOutputState(const ISimulationModuleBase* SimModule)
 	{
-		check(SimModule->GetSimType() == eSimType::Engine);
+		check(SimModule->IsSimType<class FEngineSimModule>());
 
 		FSimOutputData::FillOutputState(SimModule);
 
diff --git a/Engine/Source/Runtime/Experimental/ChaosVehicles/ChaosVehiclesCore/Private/SimModule/ModuleFactoryRegister.cpp b/Engine/Source/Runtime/Experimental/ChaosVehicles/ChaosVehiclesCore/Private/SimModule/ModuleFactoryRegister.cpp
--- a/Engine/Source/Runtime/Experimental/ChaosVehicles/ChaosVehiclesCore/Private/SimModule/ModuleFactoryRegister.cpp
+++ b/Engine/Source/Runtime/Experimental/ChaosVehicles/ChaosVehiclesCore/Private/SimModule/ModuleFactoryRegister.cpp
@@ -5,14 +5,6 @@
 
 namespace Chaos
 {
-	void FModuleFactoryRegister::RegisterFactory(int32 TypeID, TWeakPtr<IFactoryModule> InFactory)
-	{
-		if (!ContainsFactory(TypeID))
-		{
-			RegisteredFactories.Add(TypeID, InFactory);
-		}
-	}
-
 	void FModuleFactoryRegister::RegisterFactory(const FName TypeName, TWeakPtr<IFactoryModule> InFactory)
 	{
 		if (!ContainsFactory(TypeName))
@@ -23,20 +15,11 @@ namespace Chaos
 
 	void FModuleFactoryRegister::RemoveFactory(TWeakPtr<IFactoryModule> InFactory)
 	{
-		for (auto& It : RegisteredFactories)
-		{
-			if (It.Value == InFactory)
-			{
-				RegisteredFactories.Remove(It.Key);
-				return;
-			}
-		}
-
 		for (TPair<FName, TWeakPtr<IFactoryModule>> Pair : RegisteredFactoriesByName)
 		{
 			if (Pair.Value == InFactory)
-			{
-				RegisteredFactoriesByName.Remove(Pair.Key);
+			{RegisteredFactoriesByName.Remove(Pair.Key);
+				
 				return;
 			}
 		}
@@ -44,41 +27,14 @@ namespace Chaos
 
 	void FModuleFactoryRegister::Reset()
 	{
-		RegisteredFactories.Reset();
 		RegisteredFactoriesByName.Reset();
 	}
 
-	bool FModuleFactoryRegister::ContainsFactory(int32 TypeID)
-	{
-		return RegisteredFactories.Contains(TypeID);
-	}
-
-	bool FModuleFactoryRegister::ContainsFactory(const FName TypeName)
+	bool FModuleFactoryRegister::ContainsFactory(const FName TypeName) const
 	{
 		return RegisteredFactoriesByName.Contains(TypeName);
 	}
 
-	TSharedPtr<Chaos::FModuleNetData> FModuleFactoryRegister::GenerateNetData(int32 TypeID, int32 SimArrayIndex)
-	{
-		using namespace Chaos;
-
-		if (RegisteredFactories.Contains(TypeID))
-		{
-			TSharedPtr<IFactoryModule> PinnedFactory = RegisteredFactories[TypeID].Pin();
-
-			if (PinnedFactory.IsValid())
-			{
-				return PinnedFactory->GenerateNetData(SimArrayIndex);
-			}
-		}
-		else
-		{
-			UE_LOG(LogTemp, Error, TEXT("No factory registered for type '%d'"), TypeID);
-		}
-
-		return nullptr;
-	}
-
 	TSharedPtr<Chaos::FModuleNetData> FModuleFactoryRegister::GenerateNetData(const FName TypeName, const int32 SimArrayIndex)
 	{
 		using namespace Chaos;
diff --git a/Engine/Source/Runtime/Experimental/ChaosVehicles/ChaosVehiclesCore/Private/SimModule/SuspensionBaseInterface.cpp b/Engine/Source/Runtime/Experimental/ChaosVehicles/ChaosVehiclesCore/Private/SimModule/SuspensionBaseInterface.cpp
--- a/Engine/Source/Runtime/Experimental/ChaosVehicles/ChaosVehiclesCore/Private/SimModule/SuspensionBaseInterface.cpp
+++ b/Engine/Source/Runtime/Experimental/ChaosVehicles/ChaosVehiclesCore/Private/SimModule/SuspensionBaseInterface.cpp
@@ -19,11 +19,6 @@ bool FSuspensionBaseInterface::IsBehaviourType(eSimModuleTypeFlags InType) const
 	return (InType & Raycast); 
 }
 
-eSimType FSuspensionBaseInterface::GetSimType() const 
-{ 
-	return eSimType::Suspension; 
-}
-
 void FSuspensionBaseInterface::SetTargetPoint(const FVector& InTargetPoint, const FVector& InImpactNormal, float InHitDistance, bool InWheelInContact)
 {
 	TargetPos = InTargetPoint;
diff --git a/Engine/Source/Runtime/Experimental/ChaosVehicles/ChaosVehiclesCore/Private/SimModule/TorqueSimModule.cpp b/Engine/Source/Runtime/Experimental/ChaosVehicles/ChaosVehiclesCore/Private/SimModule/TorqueSimModule.cpp
--- a/Engine/Source/Runtime/Experimental/ChaosVehicles/ChaosVehiclesCore/Private/SimModule/TorqueSimModule.cpp
+++ b/Engine/Source/Runtime/Experimental/ChaosVehicles/ChaosVehiclesCore/Private/SimModule/TorqueSimModule.cpp
@@ -33,7 +33,7 @@ void FTorqueSimModule::TransmitTorque(const FSimModuleTree& ModuleTree, float Pu
 			}
 
 			ISimulationModuleBase* Module = ModuleTree.AccessSimModule(ChildIndex);
-			if (FTorqueSimModule* Interface = FTorqueSimModule::CastToTorqueInterface(Module))
+			if (FTorqueSimModule* Interface = Module->Cast<FTorqueSimModule>())
 			{
 				// push the torque value down to the children
 				// TODO: We are performing equal torque splitting here is that correct/desired?
diff --git a/Engine/Source/Runtime/Experimental/ChaosVehicles/ChaosVehiclesCore/Private/SimModule/TransmissionModule.cpp b/Engine/Source/Runtime/Experimental/ChaosVehicles/ChaosVehiclesCore/Private/SimModule/TransmissionModule.cpp
--- a/Engine/Source/Runtime/Experimental/ChaosVehicles/ChaosVehiclesCore/Private/SimModule/TransmissionModule.cpp
+++ b/Engine/Source/Runtime/Experimental/ChaosVehicles/ChaosVehiclesCore/Private/SimModule/TransmissionModule.cpp
@@ -18,11 +18,6 @@ namespace Chaos
 		, CurrentGearChangeTime(0.f)
 		, AllowedToChangeGear(true)
 	{
-		if (!FModuleFactoryRegister::Get().ContainsFactory(GetSimType()))
-		{
-			static TSharedPtr<FSimFactoryModule<FTransmissionSimModuleData>> SharedFactory = MakeShared<FSimFactoryModule<FTransmissionSimModuleData>>(GetDebugName());
-			FModuleFactoryRegister::Get().RegisterFactory(GetSimType(), SharedFactory);
-		}
 	}
 
 	void FTransmissionSimModule::Simulate(float DeltaTime, const FAllInputs& Inputs, FSimModuleTree& VehicleModuleSystem)
@@ -170,8 +165,7 @@ namespace Chaos
 
 	void FTransmissionSimModuleData::FillSimState(ISimulationModuleBase* SimModule)
 	{
-		check(SimModule->GetSimType() == eSimType::Transmission);
-		if (FTransmissionSimModule* Sim = static_cast<FTransmissionSimModule*>(SimModule))
+		if (FTransmissionSimModule* Sim = SimModule->Cast<FTransmissionSimModule>())
 		{
 			Sim->CurrentGear = CurrentGear;
 			Sim->TargetGear = TargetGear;
@@ -181,8 +175,7 @@ namespace Chaos
 
 	void FTransmissionSimModuleData::FillNetState(const ISimulationModuleBase* SimModule)
 	{
-		check(SimModule->GetSimType() == eSimType::Transmission);
-		if (const FTransmissionSimModule* Sim = static_cast<const FTransmissionSimModule*>(SimModule))
+		if (const FTransmissionSimModule* Sim = SimModule->Cast<const FTransmissionSimModule>())
 		{
 			CurrentGear = Sim->CurrentGear;
 			TargetGear = Sim->TargetGear;
@@ -210,11 +203,9 @@ namespace Chaos
 
 	void FTransmissionOutputData::FillOutputState(const ISimulationModuleBase* SimModule)
 	{
-		check(SimModule->GetSimType() == eSimType::Transmission);
-
 		FSimOutputData::FillOutputState(SimModule);
 
-		if (const FTransmissionSimModule* Sim = static_cast<const FTransmissionSimModule*>(SimModule))
+		if (const FTransmissionSimModule* Sim = SimModule->Cast<const FTransmissionSimModule>())
 		{
 			CurrentGear = Sim->CurrentGear;
 		}
diff --git a/Engine/Source/Runtime/Experimental/ChaosVehicles/ChaosVehiclesCore/Private/SimModule/WheelModule.cpp b/Engine/Source/Runtime/Experimental/ChaosVehicles/ChaosVehiclesCore/Private/SimModule/WheelModule.cpp
--- a/Engine/Source/Runtime/Experimental/ChaosVehicles/ChaosVehiclesCore/Private/SimModule/WheelModule.cpp
+++ b/Engine/Source/Runtime/Experimental/ChaosVehicles/ChaosVehiclesCore/Private/SimModule/WheelModule.cpp
@@ -17,11 +17,6 @@ namespace Chaos
 		, MassPerWheel(500.0f)
 		, SteerAngleDegrees(0.0f)
 	{
-		if (!FModuleFactoryRegister::Get().ContainsFactory(GetSimType()))
-		{
-			static TSharedPtr<FSimFactoryModule<FWheelSimModuleData>> SharedFactory = MakeShared<FSimFactoryModule<FWheelSimModuleData>>(GetDebugName());
-			FModuleFactoryRegister::Get().RegisterFactory(GetSimType(), SharedFactory);
-		}
 	}
 
 	void FWheelSimModule::Simulate(float DeltaTime, const FAllInputs& Inputs, FSimModuleTree& VehicleModuleSystem)
@@ -239,7 +234,7 @@ namespace Chaos
 
 	inline void FWheelOutputData::FillOutputState(const ISimulationModuleBase* SimModule)
 	{
-		check(SimModule->GetSimType() == eSimType::Wheel);
+		check(SimModule->IsSimType<class FWheelSimModule>());
 
 		FSimOutputData::FillOutputState(SimModule);
 
diff --git a/Engine/Source/Runtime/Experimental/ChaosVehicles/ChaosVehiclesCore/Public/SimModule/AerofoilModule.h b/Engine/Source/Runtime/Experimental/ChaosVehicles/ChaosVehiclesCore/Public/SimModule/AerofoilModule.h
--- a/Engine/Source/Runtime/Experimental/ChaosVehicles/ChaosVehiclesCore/Public/SimModule/AerofoilModule.h
+++ b/Engine/Source/Runtime/Experimental/ChaosVehicles/ChaosVehiclesCore/Public/SimModule/AerofoilModule.h
@@ -50,17 +50,15 @@ namespace Chaos
 
 	};
 
-	class CHAOSVEHICLESCORE_API FAerofoilSimModule : public ISimulationModuleBase, public TSimModuleSettings<FAerofoilSettings>
+	class CHAOSVEHICLESCORE_API FAerofoilSimModule : public ISimulationModuleBase, public TSimModuleSettings<FAerofoilSettings>, public TSimulationModuleTypeable<FAerofoilSimModule>
 	{
 	public:
-
+		DEFINE_CHAOSSIMTYPENAME(FAerofoilSimModule);
 		FAerofoilSimModule(const FAerofoilSettings& Settings);
 
 		virtual ~FAerofoilSimModule() {}
 
-		virtual TSharedPtr<FModuleNetData> GenerateNetData(int NodeArrayIndex) const { return nullptr; }
-
-		virtual eSimType GetSimType() const { return eSimType::Aerofoil; }
+		virtual TSharedPtr<FModuleNetData> GenerateNetData(const int32 NodeArrayIndex) const override { return nullptr; }
 
 		virtual const FString GetDebugName() const { return TEXT("Aerofoil"); }
 
diff --git a/Engine/Source/Runtime/Experimental/ChaosVehicles/ChaosVehiclesCore/Public/SimModule/AxleModule.h b/Engine/Source/Runtime/Experimental/ChaosVehicles/ChaosVehiclesCore/Public/SimModule/AxleModule.h
--- a/Engine/Source/Runtime/Experimental/ChaosVehicles/ChaosVehiclesCore/Public/SimModule/AxleModule.h
+++ b/Engine/Source/Runtime/Experimental/ChaosVehicles/ChaosVehiclesCore/Public/SimModule/AxleModule.h
@@ -28,13 +28,12 @@ namespace Chaos
 	/// A vehicle component that transmits torque from one source to another, i.e. from an engine or differential to wheels
 	///
 	/// </summary>
-	class CHAOSVEHICLESCORE_API FAxleSimModule : public FTorqueSimModule, public TSimModuleSettings<FAxleSettings>
+	class CHAOSVEHICLESCORE_API FAxleSimModule : public FTorqueSimModule, public TSimModuleSettings<FAxleSettings>, public TSimulationModuleTypeable<FAxleSimModule>
 	{
 	public:
+		DEFINE_CHAOSSIMTYPENAME(FAxleSimModule);
 		FAxleSimModule(const FAxleSettings& Settings);
 
-		virtual eSimType GetSimType() const { return eSimType::Axle; }
-
 		virtual const FString GetDebugName() const { return TEXT("Axle"); }
 
 		virtual void Simulate(float DeltaTime, const FAllInputs& Inputs, FSimModuleTree& VehicleModuleSystem) override;
diff --git a/Engine/Source/Runtime/Experimental/ChaosVehicles/ChaosVehiclesCore/Public/SimModule/ChassisModule.h b/Engine/Source/Runtime/Experimental/ChaosVehicles/ChaosVehiclesCore/Public/SimModule/ChassisModule.h
--- a/Engine/Source/Runtime/Experimental/ChaosVehicles/ChaosVehiclesCore/Public/SimModule/ChassisModule.h
+++ b/Engine/Source/Runtime/Experimental/ChaosVehicles/ChaosVehiclesCore/Public/SimModule/ChassisModule.h
@@ -37,14 +37,13 @@ namespace Chaos
 	/// A vehicle component that transmits torque from one source to another, i.e. from an engine or differential to wheels
 	///
 	/// </summary>
-	class CHAOSVEHICLESCORE_API FChassisSimModule : public ISimulationModuleBase, public TSimModuleSettings<FChassisSettings>
+	class CHAOSVEHICLESCORE_API FChassisSimModule : public ISimulationModuleBase, public TSimModuleSettings<FChassisSettings>, public TSimulationModuleTypeable<FChassisSimModule>
 	{
 	public:
+		DEFINE_CHAOSSIMTYPENAME(FChassisSimModule);
 		FChassisSimModule(const FChassisSettings& Settings);
 
-		virtual TSharedPtr<FModuleNetData> GenerateNetData(int NodeArrayIndex) const { return nullptr; }
-
-		virtual eSimType GetSimType() const { return eSimType::Chassis; }
+		virtual TSharedPtr<FModuleNetData> GenerateNetData(const int32 NodeArrayIndex) const override { return nullptr; }
 
 		virtual const FString GetDebugName() const { return TEXT("Chassis"); }
 
diff --git a/Engine/Source/Runtime/Experimental/ChaosVehicles/ChaosVehiclesCore/Public/SimModule/ClutchModule.h b/Engine/Source/Runtime/Experimental/ChaosVehicles/ChaosVehiclesCore/Public/SimModule/ClutchModule.h
--- a/Engine/Source/Runtime/Experimental/ChaosVehicles/ChaosVehiclesCore/Public/SimModule/ClutchModule.h
+++ b/Engine/Source/Runtime/Experimental/ChaosVehicles/ChaosVehiclesCore/Public/SimModule/ClutchModule.h
@@ -22,7 +22,7 @@ namespace Chaos
 	};
 
 
-	struct CHAOSVEHICLESCORE_API FClutchSimModuleData : public FTorqueSimModuleData
+	struct CHAOSVEHICLESCORE_API FClutchSimModuleData : public FTorqueSimModuleData, public TSimulationModuleTypeable<class FClutchSimModule, FClutchSimModuleData>
 	{
 #if !(UE_BUILD_SHIPPING || UE_BUILD_TEST)
 		FClutchSimModuleData(int NodeArrayIndex, const FString& InDebugString) : FTorqueSimModuleData(NodeArrayIndex, InDebugString) {}
@@ -30,17 +30,15 @@ namespace Chaos
 		FClutchSimModuleData(int NodeArrayIndex) : FTorqueSimModuleData(NodeArrayIndex) {}
 #endif
 
-		virtual eSimType GetType() override { return eSimType::Clutch; }
-
 		virtual void FillSimState(ISimulationModuleBase* SimModule) override
 		{
-			check(SimModule->GetSimType() == eSimType::Clutch);
+			check(SimModule->IsSimType<class FClutchSimModule>());
 			FTorqueSimModuleData::FillSimState(SimModule);
 		}
 
 		virtual void FillNetState(const ISimulationModuleBase* SimModule) override
 		{
-			check(SimModule->GetSimType() == eSimType::Clutch);
+			check(SimModule->IsSimType<class FClutchSimModule>());
 			FTorqueSimModuleData::FillNetState(SimModule);
 		}
 
@@ -56,13 +54,13 @@ namespace Chaos
 	/// Outputs - 
 	/// 
 	/// </summary>
-	class CHAOSVEHICLESCORE_API FClutchSimModule : public FTorqueSimModule, public TSimModuleSettings<FClutchSettings>
+	class CHAOSVEHICLESCORE_API FClutchSimModule : public FTorqueSimModule, public TSimModuleSettings<FClutchSettings>, public TSimulationModuleTypeable<FClutchSimModule>
 	{
 	public:
-
+		DEFINE_CHAOSSIMTYPENAME(FClutchSimModule);
 		FClutchSimModule(const FClutchSettings& Settings);
 
-		virtual TSharedPtr<FModuleNetData> GenerateNetData(int SimArrayIndex) const
+		virtual TSharedPtr<FModuleNetData> GenerateNetData(const int32 SimArrayIndex) const override
 		{
 			return MakeShared<FClutchSimModuleData>(
 				SimArrayIndex
@@ -72,8 +70,6 @@ namespace Chaos
 			);
 		}
 
-		virtual eSimType GetSimType() const { return eSimType::Clutch; }
-
 		virtual const FString GetDebugName() const { return TEXT("Clutch"); }
 
 		virtual bool GetDebugString(FString& StringOut) const override;
@@ -84,6 +80,15 @@ namespace Chaos
 
 		float ClutchValue;
 	};
-
+	
+	class CHAOSVEHICLESCORE_API FClutchSimFactory
+		: public FSimFactoryModule<FClutchSimModuleData>
+		, public TSimulationModuleTypeable<FClutchSimModule,FClutchSimFactory>
+		, public TSimFactoryAutoRegister<FClutchSimFactory>
+	
+	{
+	public:
+		FClutchSimFactory() : FSimFactoryModule(TEXT("ClutchFactory")) {}
+	};
 
 } // namespace Chaos
diff --git a/Engine/Source/Runtime/Experimental/ChaosVehicles/ChaosVehiclesCore/Public/SimModule/EngineModule.h b/Engine/Source/Runtime/Experimental/ChaosVehicles/ChaosVehiclesCore/Public/SimModule/EngineModule.h
--- a/Engine/Source/Runtime/Experimental/ChaosVehicles/ChaosVehiclesCore/Public/SimModule/EngineModule.h
+++ b/Engine/Source/Runtime/Experimental/ChaosVehicles/ChaosVehiclesCore/Public/SimModule/EngineModule.h
@@ -10,7 +10,9 @@ namespace Chaos
 	struct FAllInputs;
 	class FSimModuleTree;
 
-	struct CHAOSVEHICLESCORE_API FEngineSimModuleData : public FTorqueSimModuleData
+	struct CHAOSVEHICLESCORE_API FEngineSimModuleData
+		: public FTorqueSimModuleData
+		, public Chaos::TSimulationModuleTypeable<class FEngineSimModule,FEngineSimModuleData>
 	{
 #if !(UE_BUILD_SHIPPING || UE_BUILD_TEST)
 		FEngineSimModuleData(int NodeArrayIndex, const FString& InDebugString) : FTorqueSimModuleData(NodeArrayIndex, InDebugString) {}
@@ -18,28 +20,27 @@ namespace Chaos
 		FEngineSimModuleData(int NodeArrayIndex) : FTorqueSimModuleData(NodeArrayIndex) {}
 #endif
 
-		virtual eSimType GetType() override { return eSimType::Engine; }
-
 		virtual void FillSimState(ISimulationModuleBase* SimModule) override
 		{
-			check(SimModule->GetSimType() == eSimType::Engine);
+			check(SimModule->IsSimType<class FEngineSimModule>());
 			FTorqueSimModuleData::FillSimState(SimModule);
 		}
 
 		virtual void FillNetState(const ISimulationModuleBase* SimModule) override
 		{
-			check(SimModule->GetSimType() == eSimType::Engine);
+			check(SimModule->IsSimType<class FEngineSimModule>());
 			FTorqueSimModuleData::FillNetState(SimModule);
 		}
 
 	};
 
-	struct CHAOSVEHICLESCORE_API FEngineOutputData : public FSimOutputData
+	struct CHAOSVEHICLESCORE_API FEngineOutputData
+		: public FSimOutputData
+		, public Chaos::TSimulationModuleTypeable<class FEngineSimModule,FEngineOutputData>
 	{
 		virtual FSimOutputData* MakeNewData() override { return FEngineOutputData::MakeNew(); }
 		static FSimOutputData* MakeNew() { return new FEngineOutputData(); }
 
-		virtual eSimType GetType() override { return eSimType::Engine; }
 		virtual void FillOutputState(const ISimulationModuleBase* SimModule) override;
 		virtual void Lerp(const FSimOutputData& InCurrent, const FSimOutputData& InNext, float Alpha) override;
 #if !(UE_BUILD_SHIPPING || UE_BUILD_TEST)
@@ -82,15 +83,15 @@ namespace Chaos
 		float EngineInertia;
 	};
 
-	class CHAOSVEHICLESCORE_API FEngineSimModule : public FTorqueSimModule, public TSimModuleSettings<FEngineSettings>
+	class CHAOSVEHICLESCORE_API FEngineSimModule : public FTorqueSimModule, public TSimModuleSettings<FEngineSettings>, public TSimulationModuleTypeable<FEngineSimModule>
 	{
 	public:
-
+		DEFINE_CHAOSSIMTYPENAME(FEngineSimModule);
 		FEngineSimModule(const FEngineSettings& Settings);
 
 		virtual ~FEngineSimModule() {}
 
-		virtual TSharedPtr<FModuleNetData> GenerateNetData(int SimArrayIndex) const
+		virtual TSharedPtr<FModuleNetData> GenerateNetData(const int32 SimArrayIndex) const override
 		{
 			return MakeShared<FEngineSimModuleData>(
 				SimArrayIndex
@@ -104,8 +105,6 @@ namespace Chaos
 			return FEngineOutputData::MakeNew();
 		}
 
-		virtual eSimType GetSimType() const { return eSimType::Engine; }
-
 		virtual const FString GetDebugName() const { return TEXT("Engine"); }
 
 		virtual bool GetDebugString(FString& StringOut) const override;
@@ -124,4 +123,14 @@ namespace Chaos
 
 	};
 
+	class CHAOSVEHICLESCORE_API FEngineSimFactory
+		: public FSimFactoryModule<FEngineSimModuleData>
+		, public TSimulationModuleTypeable<FEngineSimModule,FEngineSimFactory>
+		, public TSimFactoryAutoRegister<FEngineSimFactory>
+	
+	{
+	public:
+		FEngineSimFactory() : FSimFactoryModule(TEXT("EngineSimFactory")) {}
+	};
+
 } // namespace Chaos
diff --git a/Engine/Source/Runtime/Experimental/ChaosVehicles/ChaosVehiclesCore/Public/SimModule/ModuleFactoryRegister.h b/Engine/Source/Runtime/Experimental/ChaosVehicles/ChaosVehiclesCore/Public/SimModule/ModuleFactoryRegister.h
--- a/Engine/Source/Runtime/Experimental/ChaosVehicles/ChaosVehiclesCore/Public/SimModule/ModuleFactoryRegister.h
+++ b/Engine/Source/Runtime/Experimental/ChaosVehicles/ChaosVehiclesCore/Public/SimModule/ModuleFactoryRegister.h
@@ -3,6 +3,7 @@
 #pragma once
 
 #include "CoreTypes.h"
+#include "SimulationModuleBase.h"
 #include "Templates/SharedPointer.h"
 
 namespace Chaos
@@ -19,21 +20,40 @@ namespace Chaos
 			static FModuleFactoryRegister Instance;
 			return Instance;
 		}
-
-		void RegisterFactory(int32 TypeID, TWeakPtr<IFactoryModule> InFactory);
+		
 		void RegisterFactory(const FName TypeName, TWeakPtr<IFactoryModule> InFactory);
 		void RemoveFactory(TWeakPtr<IFactoryModule> InFactory);
 		void Reset();
-		bool ContainsFactory(int32 TypeID);
-		bool ContainsFactory(const FName TypeName);
-		TSharedPtr<Chaos::FModuleNetData> GenerateNetData(int32 TypeID, int32 SimArrayIndex);
+		bool ContainsFactory(const FName TypeName) const;
 		TSharedPtr<Chaos::FModuleNetData> GenerateNetData(const FName TypeName, const int32 SimArrayIndex);
 
 	protected:
 
 		FModuleFactoryRegister() = default;
-		TMap<int32, TWeakPtr<IFactoryModule>> RegisteredFactories;
 		TMap<FName, TWeakPtr<IFactoryModule>> RegisteredFactoriesByName;
 	};
-
+	
+	template<typename _To, typename ..._Rest>
+	class TSimulationModuleTypeable;
+	//Static helper function to create and register a factory of the correct type. The returned Factory MUST be stored somewhere by the caller.
+	template<typename T, typename... Args>
+	static bool RegisterFactoryHelper(Args... args)
+	{
+		FName SimTypeName = T::StaticSimType();
+		if(SimTypeName.IsValid() == false)
+		{
+			return false;
+		}
+		if(FModuleFactoryRegister::Get().ContainsFactory(SimTypeName))
+		{
+			return true;
+		}
+		static TSharedPtr<T> SharedFactory = MakeShared<T>(args...);
+		if (SharedFactory.IsValid())
+		{
+			FModuleFactoryRegister::Get().RegisterFactory(SimTypeName, SharedFactory);
+			return true;
+		}
+		return false;
+	}
 } // namespace Chaos
diff --git a/Engine/Source/Runtime/Experimental/ChaosVehicles/ChaosVehiclesCore/Public/SimModule/MotorModule.h b/Engine/Source/Runtime/Experimental/ChaosVehicles/ChaosVehiclesCore/Public/SimModule/MotorModule.h
--- a/Engine/Source/Runtime/Experimental/ChaosVehicles/ChaosVehiclesCore/Public/SimModule/MotorModule.h
+++ b/Engine/Source/Runtime/Experimental/ChaosVehicles/ChaosVehiclesCore/Public/SimModule/MotorModule.h
@@ -30,14 +30,12 @@ namespace Chaos
 	/// A vehicle component that provides torque output based on a torque control input
 	/// The output torque is based on a square function style curve, zero at 0.0 and MaxRPM, and 1.0 at mid RPM
 	/// </summary>
-	class CHAOSVEHICLESCORE_API FMotorSimModule : public FTorqueSimModule, public TSimModuleSettings<FMotorSettings>
+	class CHAOSVEHICLESCORE_API FMotorSimModule : public FTorqueSimModule, public TSimModuleSettings<FMotorSettings>, public TSimulationModuleTypeable<FMotorSimModule>
 	{
 	public:
-
+		DEFINE_CHAOSSIMTYPENAME(FMotorSimModule);
 		FMotorSimModule(const FMotorSettings& Settings);
 
-		virtual eSimType GetSimType() const { return eSimType::Motor; }
-
 		virtual const FString GetDebugName() const { return TEXT("Motor"); }
 
 		virtual void Simulate(float DeltaTime, const FAllInputs& Inputs, FSimModuleTree& VehicleModuleSystem) override;
diff --git a/Engine/Source/Runtime/Experimental/ChaosVehicles/ChaosVehiclesCore/Public/SimModule/SimModuleTree.h b/Engine/Source/Runtime/Experimental/ChaosVehicles/ChaosVehiclesCore/Public/SimModule/SimModuleTree.h
--- a/Engine/Source/Runtime/Experimental/ChaosVehicles/ChaosVehiclesCore/Public/SimModule/SimModuleTree.h
+++ b/Engine/Source/Runtime/Experimental/ChaosVehicles/ChaosVehiclesCore/Public/SimModule/SimModuleTree.h
@@ -172,11 +172,12 @@ namespace Chaos
 			return VehicleState;
 		}
 
-		FSimModuleNode* LocateNodeByType(Chaos::eSimType InType)
+		template <typename T>
+		FSimModuleNode* LocateNodeByType()
 		{
 			for (FSimModuleNode& Node : SimulationModuleTree)
 			{
-				if (Node.SimModule && Node.SimModule->GetSimType() == InType)
+				if (Node.SimModule && Node.SimModule->IsSimType<T>())
 				{
 					return &Node;
 				}
diff --git a/Engine/Source/Runtime/Experimental/ChaosVehicles/ChaosVehiclesCore/Public/SimModule/SimulationModuleBase.h b/Engine/Source/Runtime/Experimental/ChaosVehicles/ChaosVehiclesCore/Public/SimModule/SimulationModuleBase.h
--- a/Engine/Source/Runtime/Experimental/ChaosVehicles/ChaosVehiclesCore/Public/SimModule/SimulationModuleBase.h
+++ b/Engine/Source/Runtime/Experimental/ChaosVehicles/ChaosVehiclesCore/Public/SimModule/SimulationModuleBase.h
@@ -11,6 +11,8 @@
 
 DECLARE_LOG_CATEGORY_EXTERN(LogSimulationModule, Warning, All);
 
+#define TSIMMODULETYPEABLE(clazz) Chaos::TSimulationModuleTypeableExt<clazz,&FName(#clazz)>
+
 struct CHAOSVEHICLESCORE_API FCoreModularVehicleDebugParams
 {
 	bool ShowMass = false;
@@ -101,42 +103,149 @@ namespace Chaos
 		Velocity = (1 << 3),	// requires velocity data
 	};
 
-	enum eSimType
-	{
-		Undefined = 0,
-		Chassis,		// linear/angular damping can be applied here
-		Thruster,		// applies force (can be steerable)
-		Aerofoil,		// applied drag and lift forces
-		Wheel,			// a wheel will simply roll if it has no power source
-		Suspension,		// associated with a wheel
-		Axle,			// connects more than one wheel
-		Transmission,	// gears - torque multiplier
-		Engine,			// (torque curve required) power source generates torque for wheel, axle, transmission, clutch
-		Motor,			// NOT USED YET (electric?, no torque curve required?) power source generates torque for wheel, axle, transmission, clutch
-		Clutch,			// limits the amount of torque transferred between source and destination allowing for different rotation speeds of connected axles
-		Wing,			// lift and controls aircraft roll
-		Rudder,			// controls aircraft yaw
-		Elevator,		// controls aircraft pitch
-		Propeller,		// generates thrust when connected to a motor/engine
-		TorqueSim,
-		Booster,		// applies continuous forward force to the center of mass
-		Jump,			// applies burst upward force
-	};
-
 	enum EWheelAxis
 	{
 		X,	// X forward
 		Y	// Y forward
 	};
 
+	//These two classes work in concert to provide typing information for any base classes that use ISimulationModuleBase for some base type security
+	// The FSimulationModuleTypeContainer must only be in the base class of a hierarchy.
+	// The TSimulationModuleTypeable<T> class uses multi-inheritance & CRTP to populate the FSimulationmoduleTypeContainer correctly.
+	// - All of the classes that want to be typeable/castable CRTP themselves into the chain.
+	// - All of the classes that store the shared typing information must be in the CRTP chain.
+	// Together they provide the ability to check and cast to any type within the inheritance chain.
+	
+	template<typename _To>
+	class TSimulationModuleTypeableBase
+	{
+	public:
+		TSimulationModuleTypeableBase()
+		{
+			static_cast<_To*>(this)->AddType(StaticSimTypeBase());
+		}
+		static FName StaticSimTypeBase()
+		{
+			return _To::_CHAOSTYPENAMERETRIVAL_();
+		}
+	};
+
+	template<typename _To, typename ..._Rest>
+	class TSimulationModuleTypeable;
+	
+	template<typename _To>
+	class TSimulationModuleTypeable<_To>
+	{
+	public:
+		TSimulationModuleTypeable()
+		{
+			static_cast<_To*>(this)->AddType(StaticSimType());
+		}
+		static FName StaticSimType()
+		{
+			return _To::_CHAOSTYPENAMERETRIVAL_();
+		}
+	};
+	
+	template<typename _To, typename _From>
+	class TSimulationModuleTypeable<_To,_From>
+	{
+	public:
+		TSimulationModuleTypeable()
+		{
+			static_cast<_From*>(this)->AddType(StaticSimType());
+		}
+		static FName StaticSimType()
+		{
+			return _To::_CHAOSTYPENAMERETRIVAL_();
+		}
+		static FName RecurseSimType()
+		{
+			return StaticSimType();
+		}
+	};
+
+	
+	template<class T, class = void>
+	struct TSimModuleTypeIsRecursive
+	: std::false_type
+	{};
+
+	template<class T>
+	struct TSimModuleTypeIsRecursive<T, std::enable_if_t<std::is_invocable_r<FName, decltype(T::RecurseSimType)>::value>>
+	: std::integral_constant<bool, true>
+	{};
+
+	template<class T>
+	constexpr bool TSimModuleTypeIsRecursive_v = TSimModuleTypeIsRecursive<T>::value;
+	
+	class FSimulationModuleTypeContainer
+	{
+	public:
+		TSet<FName> MyTypes;
+		FName MostRecentAdd = NAME_None;
+		void AddType(FName InType)
+		{
+			MyTypes.Emplace(InType);
+			MostRecentAdd = InType;
+		}
+		bool IsSimType(FName InType) const
+		{
+			return MostRecentAdd == InType || MyTypes.Contains(InType);
+		}
+		FName GetSimType() const
+		{
+			return MostRecentAdd;
+		}
+		template<typename U>
+		bool IsSimType() const
+		{
+			if constexpr (TSimModuleTypeIsRecursive_v<U>)
+			{
+				return IsSimType(U::RecurseSimType()); 
+			}
+			else
+			{
+				return IsSimType(TSimulationModuleTypeable<U>::StaticSimType()); 
+			}
+		}
+		template<typename U>
+		U* Cast()
+		{
+			if(IsSimType<std::remove_const_t<U>>())
+			{
+				return static_cast<U*>(this);
+			}
+			return nullptr;
+		}
+		template<typename U>
+		const U* Cast() const
+		{
+			if(IsSimType<std::remove_const_t<U>>())
+			{
+				return static_cast<const U*>(this);
+			}
+			return nullptr;
+		}
+	};
+	
+
+#define DEFINE_CHAOSSIMTYPENAME(cls) \
+	 static FName _CHAOSTYPENAMERETRIVAL_() { return FName(#cls); }
+	
+	//FName TSimulationModuleTypeName<class cls>::SimModuleTypeName = FName(#cls); 
+
+	
 	/**
 	 * Interface base class for all simulation module building blocks
 	 */
-	class CHAOSVEHICLESCORE_API ISimulationModuleBase
+	class CHAOSVEHICLESCORE_API ISimulationModuleBase : public FSimulationModuleTypeContainer, public TSimulationModuleTypeableBase<ISimulationModuleBase>
 	{
 	friend FSimOutputData;
-
+		
+		
 	public:
+		DEFINE_CHAOSSIMTYPENAME(ISimulationModuleBase);
 		const static int INVALID_IDX = -1;
 
 		ISimulationModuleBase()
@@ -159,7 +268,6 @@ namespace Chaos
 
 		const int GetGuid() { return Guid; }
 		void SetGuid(int GuidIn) { Guid = GuidIn; }
-
 		/**
 		* Get the friendly name for this module, primarily for logging & debugging module tree
 		 */
@@ -170,16 +278,6 @@ namespace Chaos
 		 */
 		virtual bool IsBehaviourType(eSimModuleTypeFlags InType) const = 0;
 
-		/**
-		 * The specific simulation type
-		 */
-		virtual eSimType GetSimType() const = 0;
-
-		/**
-		 * The specific simulation type
-		 */
-		virtual FName GetSimTypeName() const { return NAME_None; }
-
 		/**
 		 * Is Module active and simulating
 		 */
@@ -331,7 +429,7 @@ namespace Chaos
 		const FVector& GetAppliedForce() { return AppliedForce; }
 
 		// this is the replication datas
-		virtual TSharedPtr<FModuleNetData> GenerateNetData(int NodeArrayIndex) const = 0;
+		virtual TSharedPtr<FModuleNetData> GenerateNetData(const int32 NodeArrayIndex) const = 0;
 
 		virtual FSimOutputData* GenerateOutputData() const { return nullptr; }
 
@@ -374,8 +472,9 @@ namespace Chaos
 	/**
 	* Interface base class for all module network serialization
 	*/
-	struct CHAOSVEHICLESCORE_API FModuleNetData
+	struct CHAOSVEHICLESCORE_API FModuleNetData :public FSimulationModuleTypeContainer, public TSimulationModuleTypeableBase<FModuleNetData>
 	{
+		DEFINE_CHAOSSIMTYPENAME(FModuleNetData);
 		FModuleNetData(int InSimArrayIndex, const FString& InDebugString = FString())
 			: SimArrayIndex(InSimArrayIndex)
 #if !(UE_BUILD_SHIPPING || UE_BUILD_TEST)
@@ -385,8 +484,6 @@ namespace Chaos
 
 		virtual ~FModuleNetData() {}
 
-		virtual eSimType GetType() = 0;
-		virtual FName GetTypeName() const { return NAME_None; };
 		virtual void Serialize(FArchive& Ar) = 0;
 		virtual void FillNetState(const ISimulationModuleBase* SimModule) = 0;
 		virtual void FillSimState(ISimulationModuleBase* SimModule) = 0;
@@ -408,14 +505,13 @@ namespace Chaos
 		static uint16 AnimatePosition = 0x00000001;
 		static uint16 AnimateRotation = 0x00000002;
 	}
-
-	struct CHAOSVEHICLESCORE_API FSimOutputData
+	
+	struct CHAOSVEHICLESCORE_API FSimOutputData : public FSimulationModuleTypeContainer, public TSimulationModuleTypeableBase<FSimOutputData>
 	{
+		DEFINE_CHAOSSIMTYPENAME(FSimOutputData);
 		FSimOutputData() = default;
 		virtual ~FSimOutputData() {}
-
-		virtual eSimType GetType() = 0;
-		virtual FName GetTypeName() const { return NAME_None; };
+		
 		virtual bool IsEnabled() { return bEnabled; }
 		virtual FSimOutputData* MakeNewData() = 0;
 		virtual void FillOutputState(const ISimulationModuleBase* SimModule);
@@ -433,16 +529,24 @@ namespace Chaos
 #endif
 	};
 
-	class CHAOSVEHICLESCORE_API IFactoryModule
+	class CHAOSVEHICLESCORE_API IFactoryModule : public FSimulationModuleTypeContainer, public TSimulationModuleTypeableBase<IFactoryModule>
 	{
 	public:
+		DEFINE_CHAOSSIMTYPENAME(IFactoryModule);
 		virtual ~IFactoryModule() {}
-		virtual TSharedPtr<Chaos::FModuleNetData> GenerateNetData(int32 SimArrayIndex) const = 0;
+		virtual TSharedPtr<Chaos::FModuleNetData> GenerateNetData(const int32 SimArrayIndex) const = 0;
 
 	};
 
+	template<typename _To>
+	class TSimFactoryAutoRegister
+	{
+	private:
+		inline static bool bSimFactoryRegistered = RegisterFactoryHelper<_To>();
+	};
+
 	template <typename T>
-	class CHAOSVEHICLESCORE_API FSimFactoryModule : public IFactoryModule
+	class FSimFactoryModule : public IFactoryModule
 	{
 	public:
 		FSimFactoryModule(const FString& DebugNameIn)
@@ -452,8 +556,7 @@ namespace Chaos
 #endif
 		}
 
-
-		TSharedPtr<Chaos::FModuleNetData> GenerateNetData(int32 SimArrayIndex) const
+		TSharedPtr<Chaos::FModuleNetData> GenerateNetData(const int32 SimArrayIndex) const override
 		{
 			return MakeShared<T>(
 				SimArrayIndex
@@ -468,5 +571,4 @@ namespace Chaos
 #endif
 
 	};
-
 } // namespace Chaos
diff --git a/Engine/Source/Runtime/Experimental/ChaosVehicles/ChaosVehiclesCore/Public/SimModule/SuspensionBaseInterface.h b/Engine/Source/Runtime/Experimental/ChaosVehicles/ChaosVehiclesCore/Public/SimModule/SuspensionBaseInterface.h
--- a/Engine/Source/Runtime/Experimental/ChaosVehicles/ChaosVehiclesCore/Public/SimModule/SuspensionBaseInterface.h
+++ b/Engine/Source/Runtime/Experimental/ChaosVehicles/ChaosVehiclesCore/Public/SimModule/SuspensionBaseInterface.h
@@ -27,17 +27,16 @@ struct CHAOSVEHICLESCORE_API FSpringTrace
 };
 
 
-class CHAOSVEHICLESCORE_API FSuspensionBaseInterface : public ISimulationModuleBase
+class CHAOSVEHICLESCORE_API FSuspensionBaseInterface : public ISimulationModuleBase, public TSimulationModuleTypeable<FSuspensionBaseInterface>
 {
 public:
+	DEFINE_CHAOSSIMTYPENAME(FSuspensionBaseInterface);
 	FSuspensionBaseInterface();
 
 	virtual ~FSuspensionBaseInterface() {}
 
 	virtual bool IsBehaviourType(eSimModuleTypeFlags InType) const override;
 
-	virtual eSimType GetSimType() const override;
-
 	virtual float GetMaxSpringLength() const = 0;
 	virtual float GetSpringLength() const = 0;
 	virtual void SetSpringLength(float InLength, float WheelRadius) = 0;
diff --git a/Engine/Source/Runtime/Experimental/ChaosVehicles/ChaosVehiclesCore/Public/SimModule/ThrusterModule.h b/Engine/Source/Runtime/Experimental/ChaosVehicles/ChaosVehiclesCore/Public/SimModule/ThrusterModule.h
--- a/Engine/Source/Runtime/Experimental/ChaosVehicles/ChaosVehiclesCore/Public/SimModule/ThrusterModule.h
+++ b/Engine/Source/Runtime/Experimental/ChaosVehicles/ChaosVehiclesCore/Public/SimModule/ThrusterModule.h
@@ -44,14 +44,13 @@ namespace Chaos
 	/// A vehicle component that transmits torque from one source to another, i.e. from an engine or differential to wheels
 	///
 	/// </summary>
-	class CHAOSVEHICLESCORE_API FThrusterSimModule : public ISimulationModuleBase, public TSimModuleSettings<FThrusterSettings>
+	class CHAOSVEHICLESCORE_API FThrusterSimModule : public ISimulationModuleBase, public TSimModuleSettings<FThrusterSettings>, public TSimulationModuleTypeable<FThrusterSimModule>
 	{
 	public:
+		DEFINE_CHAOSSIMTYPENAME(FThrusterSimModule);
 		FThrusterSimModule(const FThrusterSettings& Settings);
 
-		virtual TSharedPtr<FModuleNetData> GenerateNetData(int NodeArrayIndex) const { return nullptr; }
-
-		virtual eSimType GetSimType() const { return eSimType::Thruster; }
+		virtual TSharedPtr<FModuleNetData> GenerateNetData(const int32 NodeArrayIndex) const override { return nullptr; }
 
 		virtual const FString GetDebugName() const { return TEXT("Thruster"); }
 
diff --git a/Engine/Source/Runtime/Experimental/ChaosVehicles/ChaosVehiclesCore/Public/SimModule/TorqueSimModule.h b/Engine/Source/Runtime/Experimental/ChaosVehicles/ChaosVehiclesCore/Public/SimModule/TorqueSimModule.h
--- a/Engine/Source/Runtime/Experimental/ChaosVehicles/ChaosVehiclesCore/Public/SimModule/TorqueSimModule.h
+++ b/Engine/Source/Runtime/Experimental/ChaosVehicles/ChaosVehiclesCore/Public/SimModule/TorqueSimModule.h
@@ -9,14 +9,15 @@ namespace Chaos
 {
 	class FTorqueSimModule;
 
-	struct CHAOSVEHICLESCORE_API FTorqueSimModuleData : public FModuleNetData
+	struct CHAOSVEHICLESCORE_API FTorqueSimModuleData
+		: public FModuleNetData
+		, public Chaos::TSimulationModuleTypeable<FTorqueSimModule,FTorqueSimModuleData>
 	{
 #if !(UE_BUILD_SHIPPING || UE_BUILD_TEST)
 		FTorqueSimModuleData(int NodeArrayIndex, const FString& InDebugString) : FModuleNetData(NodeArrayIndex, InDebugString) {}
 #else
 		FTorqueSimModuleData(int NodeArrayIndex) : FModuleNetData(NodeArrayIndex) {}
 #endif
-		virtual eSimType GetType() override { return eSimType::TorqueSim; }
 
 		virtual void FillSimState(ISimulationModuleBase* SimModule) override;
 
@@ -45,11 +46,12 @@ namespace Chaos
 	};
 
 
-	class CHAOSVEHICLESCORE_API FTorqueSimModule : public ISimulationModuleBase
+	class CHAOSVEHICLESCORE_API FTorqueSimModule : public ISimulationModuleBase, public TSimulationModuleTypeable<FTorqueSimModule>
 	{
 		friend FTorqueSimModuleData;
 
 	public:
+		DEFINE_CHAOSSIMTYPENAME(FTorqueSimModule);
 		FTorqueSimModule()
 			: DriveTorque(0.f)
 			, LoadTorque(0.f)
@@ -64,8 +66,6 @@ namespace Chaos
 		 */
 		virtual bool IsBehaviourType(eSimModuleTypeFlags InType) const { return (InType & TorqueBased); }
 
-		virtual eSimType GetSimType() const { return eSimType::TorqueSim; }
-
 		void SetDriveTorque(float TorqueIn) { DriveTorque = TorqueIn; }
 		float GetDriveTorque() const { return DriveTorque; }
 
@@ -118,9 +118,11 @@ namespace Chaos
 		float AngularPosition;
 	};
 
-	class CHAOSVEHICLESCORE_API FWheelBaseInterface : public FTorqueSimModule
+	class CHAOSVEHICLESCORE_API FWheelBaseInterface : public FTorqueSimModule, public TSimulationModuleTypeable<FWheelBaseInterface>
 	{
 	public:
+		DEFINE_CHAOSSIMTYPENAME(FWheelBaseInterface);
+		
 		FWheelBaseInterface()
 			: SuspensionSimTreeIndex(ISimulationModuleBase::INVALID_IDX)
 			, SurfaceFriction(1.0f)
diff --git a/Engine/Source/Runtime/Experimental/ChaosVehicles/ChaosVehiclesCore/Public/SimModule/TransmissionModule.h b/Engine/Source/Runtime/Experimental/ChaosVehicles/ChaosVehiclesCore/Public/SimModule/TransmissionModule.h
--- a/Engine/Source/Runtime/Experimental/ChaosVehicles/ChaosVehiclesCore/Public/SimModule/TransmissionModule.h
+++ b/Engine/Source/Runtime/Experimental/ChaosVehicles/ChaosVehiclesCore/Public/SimModule/TransmissionModule.h
@@ -7,10 +7,13 @@
 
 namespace Chaos
 {
+	class FTransmissionSimModule;
 	struct FAllInputs;
 	class FSimModuleTree;
 
-	struct CHAOSVEHICLESCORE_API FTransmissionSimModuleData : public FModuleNetData
+	struct CHAOSVEHICLESCORE_API FTransmissionSimModuleData
+		: public FModuleNetData
+		, public Chaos::TSimulationModuleTypeable<FTransmissionSimModule,FTransmissionSimModuleData>
 	{
 #if !(UE_BUILD_SHIPPING || UE_BUILD_TEST)
 		FTransmissionSimModuleData(int NodeArrayIndex, const FString& InDebugString) : FModuleNetData(NodeArrayIndex, InDebugString) {}
@@ -18,8 +21,6 @@ namespace Chaos
 		FTransmissionSimModuleData(int NodeArrayIndex) : FModuleNetData(NodeArrayIndex) {}
 #endif
 
-		virtual eSimType GetType() override { return eSimType::Transmission; }
-
 		virtual void FillSimState(ISimulationModuleBase* SimModule) override;
 
 		virtual void FillNetState(const ISimulationModuleBase* SimModule) override;
@@ -42,12 +43,13 @@ namespace Chaos
 		float CurrentGearChangeTime = 0.0f;
 	};
 
-	struct CHAOSVEHICLESCORE_API FTransmissionOutputData : public FSimOutputData
+	struct CHAOSVEHICLESCORE_API FTransmissionOutputData
+		: public FSimOutputData
+		, public Chaos::TSimulationModuleTypeable<FTransmissionSimModule,FTransmissionOutputData>
 	{
 		virtual FSimOutputData* MakeNewData() override { return FTransmissionOutputData::MakeNew(); }
 		static FSimOutputData* MakeNew() { return new FTransmissionOutputData(); }
-
-		virtual eSimType GetType() override { return eSimType::Transmission; }
+		
 		virtual void FillOutputState(const ISimulationModuleBase* SimModule) override;
 		virtual void Lerp(const FSimOutputData& InCurrent, const FSimOutputData& InNext, float Alpha) override;
 
@@ -98,16 +100,16 @@ namespace Chaos
 		bool AutoReverse;				// Arcade handling - holding Brake switches into reverse after vehicle has stopped
 	};
 
-	class CHAOSVEHICLESCORE_API FTransmissionSimModule : public FTorqueSimModule, public TSimModuleSettings<FTransmissionSettings>
+	class CHAOSVEHICLESCORE_API FTransmissionSimModule : public FTorqueSimModule, public TSimModuleSettings<FTransmissionSettings>, public TSimulationModuleTypeable<FTransmissionSimModule>
 	{
 		friend FTransmissionSimModuleData;
 		friend FTransmissionOutputData;
 
 	public:
-
+		DEFINE_CHAOSSIMTYPENAME(FTransmissionSimModule);
 		FTransmissionSimModule(const FTransmissionSettings& Settings);
 
-		virtual TSharedPtr<FModuleNetData> GenerateNetData(int SimArrayIndex) const
+		virtual TSharedPtr<FModuleNetData> GenerateNetData(const int32 SimArrayIndex) const override
 		{
 			return MakeShared<FTransmissionSimModuleData>(
 				SimArrayIndex
@@ -122,8 +124,6 @@ namespace Chaos
 			return FTransmissionOutputData::MakeNew();
 		}
 
-		virtual eSimType GetSimType() const { return eSimType::Transmission; }
-
 		virtual const FString GetDebugName() const { return TEXT("Transmission"); }
 
 		virtual bool GetDebugString(FString& StringOut) const override;
@@ -176,6 +176,16 @@ namespace Chaos
 		bool AllowedToChangeGear; // conditions are ok for an automatic gear change
 
 	};
+	
+	class CHAOSVEHICLESCORE_API FTransmissionSimFactory
+		: public FSimFactoryModule<FTransmissionSimModuleData>
+		, public TSimulationModuleTypeable<FTransmissionSimModule,FTransmissionSimFactory>
+		, public TSimFactoryAutoRegister<FTransmissionSimFactory>
+	
+	{
+	public:
+		FTransmissionSimFactory() : FSimFactoryModule(TEXT("TransmissionSimFactory")) {}
+	};
 
 
 } // namespace Chaos
diff --git a/Engine/Source/Runtime/Experimental/ChaosVehicles/ChaosVehiclesCore/Public/SimModule/WheelModule.h b/Engine/Source/Runtime/Experimental/ChaosVehicles/ChaosVehiclesCore/Public/SimModule/WheelModule.h
--- a/Engine/Source/Runtime/Experimental/ChaosVehicles/ChaosVehiclesCore/Public/SimModule/WheelModule.h
+++ b/Engine/Source/Runtime/Experimental/ChaosVehicles/ChaosVehiclesCore/Public/SimModule/WheelModule.h
@@ -11,7 +11,9 @@ namespace Chaos
 	class FClusterUnionPhysicsProxy;
 	struct FModuleNetData;
 
-	struct CHAOSVEHICLESCORE_API FWheelSimModuleData : public FTorqueSimModuleData
+	struct CHAOSVEHICLESCORE_API FWheelSimModuleData
+		: public FTorqueSimModuleData
+		, public Chaos::TSimulationModuleTypeable<class FWheelSimModule,FWheelSimModuleData>
 	{
 #if !(UE_BUILD_SHIPPING || UE_BUILD_TEST)
 		FWheelSimModuleData(int NodeArrayIndex, const FString& InDebugString) : FTorqueSimModuleData(NodeArrayIndex, InDebugString) {}
@@ -19,28 +21,27 @@ namespace Chaos
 		FWheelSimModuleData(int NodeArrayIndex) : FTorqueSimModuleData(NodeArrayIndex) {}
 #endif
 
-		virtual eSimType GetType() override { return eSimType::Wheel; }
-
 		virtual void FillSimState(ISimulationModuleBase* SimModule) override
 		{
-			check(SimModule->GetSimType() == eSimType::Wheel);
+			check(SimModule->IsSimType<class FWheelSimModule>());
 			FTorqueSimModuleData::FillSimState(SimModule);
 		}
 
 		virtual void FillNetState(const ISimulationModuleBase* SimModule) override
 		{
-			check(SimModule->GetSimType() == eSimType::Wheel);
+			check(SimModule->IsSimType<class FWheelSimModule>());
 			FTorqueSimModuleData::FillNetState(SimModule);
 		}
 
 	};
 
-	struct CHAOSVEHICLESCORE_API FWheelOutputData : public FSimOutputData
+	struct CHAOSVEHICLESCORE_API FWheelOutputData
+		: public FSimOutputData
+		, public Chaos::TSimulationModuleTypeable<class FWheelSimModule,FWheelOutputData>
 	{
 		virtual FSimOutputData* MakeNewData() override { return FWheelOutputData::MakeNew(); }
 		static FSimOutputData* MakeNew() { return new FWheelOutputData(); }
-
-		virtual eSimType GetType() override { return eSimType::Wheel; }
+		
 		virtual void FillOutputState(const ISimulationModuleBase* SimModule) override;
 		virtual void Lerp(const FSimOutputData& InCurrent, const FSimOutputData& InNext, float Alpha) override;
 
@@ -115,15 +116,14 @@ namespace Chaos
 		bool ReverseDirection;
 	};
 
-	class CHAOSVEHICLESCORE_API FWheelSimModule : public FWheelBaseInterface, public TSimModuleSettings<FWheelSettings>
+	class CHAOSVEHICLESCORE_API FWheelSimModule : public FWheelBaseInterface, public TSimModuleSettings<FWheelSettings>, public TSimulationModuleTypeable<FWheelSimModule>
 	{
 		friend FWheelOutputData;
-
 	public:
-
+		DEFINE_CHAOSSIMTYPENAME(FWheelSimModule);
 		FWheelSimModule(const FWheelSettings& Settings);
 
-		virtual TSharedPtr<FModuleNetData> GenerateNetData(int SimArrayIndex) const
+		virtual TSharedPtr<FModuleNetData> GenerateNetData(const int32 SimArrayIndex) const override
 		{
 			return MakeShared<FWheelSimModuleData>(
 				SimArrayIndex
@@ -138,8 +138,6 @@ namespace Chaos
 			return FWheelOutputData::MakeNew();
 		}
 
-		virtual eSimType GetSimType() const { return eSimType::Wheel; }
-
 		virtual const FString GetDebugName() const { return TEXT("Wheel"); }
 
 		virtual bool GetDebugString(FString& StringOut) const override;
@@ -188,5 +186,15 @@ namespace Chaos
 		float SlipAngle;
 	};
 
+	
+	class CHAOSVEHICLESCORE_API FWheelSimFactory
+			: public FSimFactoryModule<FWheelSimModuleData>
+			, public TSimulationModuleTypeable<FWheelSimModule,FWheelSimFactory>
+			, public TSimFactoryAutoRegister<FWheelSimFactory>
+	
+	{
+	public:
+		FWheelSimFactory() : FSimFactoryModule(TEXT("WheelSimFactory")) {}
+	};
 
 } // namespace Chaos
