[ModularVehicles] Check validity of pointer before use.

diff --git a/Engine/Plugins/Experimental/ChaosModularVehicle/Source/ChaosModularVehicleEngine/Private/ChaosModularVehicle/ModularVehicleBaseComponent.cpp b/Engine/Plugins/Experimental/ChaosModularVehicle/Source/ChaosModularVehicleEngine/Private/ChaosModularVehicle/ModularVehicleBaseComponent.cpp
--- a/Engine/Plugins/Experimental/ChaosModularVehicle/Source/ChaosModularVehicleEngine/Private/ChaosModularVehicle/ModularVehicleBaseComponent.cpp
+++ b/Engine/Plugins/Experimental/ChaosModularVehicle/Source/ChaosModularVehicleEngine/Private/ChaosModularVehicle/ModularVehicleBaseComponent.cpp
@@ -752,40 +752,46 @@ void UModularVehicleBaseComponent::ActionTreeUpdates(Chaos::FSimTreeUpdates* Nex
 			if (VehicleSimulationPT)
 			{
 				TUniquePtr<Chaos::FSimModuleTree>& SimModuleTree = VehicleSimulationPT->AccessSimComponentTree();
-				SimModuleTree->AppendTreeUpdates(NextTreeUpdates);
-				FModularVehicleBuilder::FixupTreeLinks(SimModuleTree);
-
-#if !(UE_BUILD_SHIPPING || UE_BUILD_TEST)
-				if (bModularVehicle_DumpModuleTreeStructure_Enabled)
+				if(SimModuleTree.IsValid())
 				{
-					UE_LOG(LogTemp, Warning, TEXT("SimTreeModules:") );
-					for (int I = 0; I < SimModuleTree->GetNumNodes(); I++)
+					SimModuleTree->AppendTreeUpdates(NextTreeUpdates);
+					FModularVehicleBuilder::FixupTreeLinks(SimModuleTree);
+
+	#if !(UE_BUILD_SHIPPING || UE_BUILD_TEST)
+					if (bModularVehicle_DumpModuleTreeStructure_Enabled)
 					{
-						if (Chaos::ISimulationModuleBase* Module = SimModuleTree->GetNode(I).SimModule)
+						UE_LOG(LogTemp, Warning, TEXT("SimTreeModules:") );
+						for (int I = 0; I < SimModuleTree->GetNumNodes(); I++)
 						{
-							FString String;
-							Module->GetDebugString(String);
-							UE_LOG(LogTemp, Warning, TEXT("..%s"), *String);
+							if (Chaos::ISimulationModuleBase* Module = SimModuleTree->GetNode(I).SimModule)
+							{
+								FString String;
+								Module->GetDebugString(String);
+								UE_LOG(LogTemp, Warning, TEXT("..%s"), *String);
+							}
 						}
 					}
-				}
-#endif // !(UE_BUILD_SHIPPING || UE_BUILD_TEST)
+	#endif // !(UE_BUILD_SHIPPING || UE_BUILD_TEST)
 
-				// Network replication data needs to be updated, this is currently studily slow
-				if (NetworkPhysicsComponent)
-				{
-					TSharedPtr<Chaos::FBaseRewindHistory>& History = NetworkPhysicsComponent->GetStateHistory_Internal();
-					Chaos::TDataRewindHistory<FNetworkModularVehicleStates>* StateHistory = static_cast<Chaos::TDataRewindHistory<FNetworkModularVehicleStates>*>(History.Get());
-					if (StateHistory)
+					// Network replication data needs to be updated, this is currently studily slow
+					if (NetworkPhysicsComponent)
 					{
-						// #TODO: we are rebuilding from scratch every time there is a single change, there must be a better way!
-						// not sure of it is safe to update the data at this time?
-						for (int I = 0; I < StateHistory->GetDataHistory().Num(); I++)
+						TSharedPtr<Chaos::FBaseRewindHistory>& History = NetworkPhysicsComponent->GetStateHistory_Internal();
+						Chaos::TDataRewindHistory<FNetworkModularVehicleStates>* StateHistory = static_cast<Chaos::TDataRewindHistory<FNetworkModularVehicleStates>*>(History.Get());
+						if (StateHistory)
 						{
-							FNetworkModularVehicleStates& State = StateHistory->GetDataHistory()[I];
-							State.ModuleData.Empty();
-
-							VehicleSimulationPT->AccessSimComponentTree()->GenerateReplicationStructure(State.ModuleData);
+							// #TODO: we are rebuilding from scratch every time there is a single change, there must be a better way!
+							// not sure of it is safe to update the data at this time?
+							for (int I = 0; I < StateHistory->GetDataHistory().Num(); I++)
+							{
+								FNetworkModularVehicleStates& State = StateHistory->GetDataHistory()[I];
+								State.ModuleData.Empty();
+								TUniquePtr<Chaos::FSimModuleTree>& InnerSimModuleTree = VehicleSimulationPT->AccessSimComponentTree();
+								if(InnerSimModuleTree.IsValid())
+								{
+									InnerSimModuleTree->GenerateReplicationStructure(State.ModuleData);
+								}
+							}
 						}
 					}
 				}
@@ -885,18 +891,21 @@ void UModularVehicleBaseComponent::RemoveComponentFromSimulation(UPrimitiveCompo
 		}
 
 		TUniquePtr<Chaos::FSimModuleTree>& SimModuleTree = VehicleSimulationPT->AccessSimComponentTree();
-		for (const Chaos::FPendingModuleDeletions& TreeUpdate : LatestTreeUpdates.GetDeletedModules())
+		if(SimModuleTree.IsValid())
 		{
-			for (int Index = 0; Index < SimModuleTree->GetNumNodes(); Index++)
+			for (const Chaos::FPendingModuleDeletions& TreeUpdate : LatestTreeUpdates.GetDeletedModules())
 			{
-				if (Chaos::ISimulationModuleBase* SimModule = SimModuleTree->GetNode(Index).SimModule)
+				for (int Index = 0; Index < SimModuleTree->GetNumNodes(); Index++)
 				{
-					if (SimModule->GetGuid() == TreeUpdate.Guid)
+					if (Chaos::ISimulationModuleBase* SimModule = SimModuleTree->GetNode(Index).SimModule)
 					{
-						SimModule->SetAnimationEnabled(false);
-						SimModule->SetStateFlags(Chaos::eSimModuleState::Disabled);
-						SimModule->OnTermination_External();
-						break;
+						if (SimModule->GetGuid() == TreeUpdate.Guid)
+						{
+							SimModule->SetAnimationEnabled(false);
+							SimModule->SetStateFlags(Chaos::eSimModuleState::Disabled);
+							SimModule->OnTermination_External();
+							break;
+						}
 					}
 				}
 			}
@@ -908,7 +917,8 @@ void UModularVehicleBaseComponent::RemoveComponentFromSimulation(UPrimitiveCompo
 			{
 				if (VehicleSimulationPT)
 				{
-					if(TUniquePtr<Chaos::FSimModuleTree>& SimModuleTree = VehicleSimulationPT->AccessSimComponentTree())
+					TUniquePtr<Chaos::FSimModuleTree>& SimModuleTree = VehicleSimulationPT->AccessSimComponentTree();
+					if(SimModuleTree.IsValid())
 					{
 						SimModuleTree->AppendTreeUpdates(LatestTreeUpdates);
 					}
