[ModularVehicles] Resolve use after free issues.

diff --git a/Engine/Plugins/Experimental/ChaosModularVehicle/Source/ChaosModularVehicleEngine/Private/ChaosModularVehicle/ChaosSimModuleManagerAsyncCallback.cpp b/Engine/Plugins/Experimental/ChaosModularVehicle/Source/ChaosModularVehicleEngine/Private/ChaosModularVehicle/ChaosSimModuleManagerAsyncCallback.cpp
--- a/Engine/Plugins/Experimental/ChaosModularVehicle/Source/ChaosModularVehicleEngine/Private/ChaosModularVehicle/ChaosSimModuleManagerAsyncCallback.cpp
+++ b/Engine/Plugins/Experimental/ChaosModularVehicle/Source/ChaosModularVehicleEngine/Private/ChaosModularVehicle/ChaosSimModuleManagerAsyncCallback.cpp
@@ -197,7 +197,7 @@ void FModularVehicleAsyncInput::OnContactModification(Chaos::FCollisionContactMo
 
 void FModularVehicleAsyncInput::ApplyDeferredForces() const
 {
-	if (Vehicle && Proxy)
+	if (Vehicle && Proxy && Vehicle->VehicleSimulationPT)
 	{
 		if (Proxy->GetType() == EPhysicsProxyType::ClusterUnionProxy)
 		{
diff --git a/Engine/Plugins/Experimental/ChaosModularVehicle/Source/ChaosModularVehicleEngine/Private/ChaosModularVehicle/ModularVehicleBaseComponent.cpp b/Engine/Plugins/Experimental/ChaosModularVehicle/Source/ChaosModularVehicleEngine/Private/ChaosModularVehicle/ModularVehicleBaseComponent.cpp
--- a/Engine/Plugins/Experimental/ChaosModularVehicle/Source/ChaosModularVehicleEngine/Private/ChaosModularVehicle/ModularVehicleBaseComponent.cpp
+++ b/Engine/Plugins/Experimental/ChaosModularVehicle/Source/ChaosModularVehicleEngine/Private/ChaosModularVehicle/ModularVehicleBaseComponent.cpp
@@ -749,7 +749,7 @@ void UModularVehicleBaseComponent::ActionTreeUpdates(Chaos::FSimTreeUpdates* Nex
 			{
 	Solver->EnqueueCommandImmediate([Proxy, this, NextTreeUpdates = *NextTreeUpdates]() mutable
 		{
-			if (VehicleSimulationPT)
+			if (IsValid(this) && bPhysicsStateCreated && VehicleSimulationPT)
 			{
 				TUniquePtr<Chaos::FSimModuleTree>& SimModuleTree = VehicleSimulationPT->AccessSimComponentTree();
 				if(SimModuleTree.IsValid())
@@ -915,7 +915,7 @@ void UModularVehicleBaseComponent::RemoveComponentFromSimulation(UPrimitiveCompo
 		Chaos::FPBDRigidsSolver* Solver = Proxy->GetSolver<Chaos::FPBDRigidsSolver>();
 		Solver->EnqueueCommandImmediate([Proxy, this, LatestTreeUpdates = LatestTreeUpdates]() mutable
 			{
-				if (VehicleSimulationPT)
+				if (IsValid(this) && bPhysicsStateCreated && VehicleSimulationPT)
 				{
 					TUniquePtr<Chaos::FSimModuleTree>& SimModuleTree = VehicleSimulationPT->AccessSimComponentTree();
 					if(SimModuleTree.IsValid())
@@ -1024,19 +1024,16 @@ void UModularVehicleBaseComponent::DestroyVehicleSim()
 				SimManager->RemoveVehicle(this);
 			}
 
-			LocalSolver->EnqueueCommandImmediate([this, PhysScene = PhysScene]() mutable
-				{
-					if (PVehicleOutput.IsValid())
-					{
-						PVehicleOutput.Reset(nullptr);
-					}
+			if (PVehicleOutput.IsValid())
+			{
+				PVehicleOutput.Reset(nullptr);
+			}
 
-					if (VehicleSimulationPT)
-					{
-						VehicleSimulationPT->Terminate();
-						VehicleSimulationPT.Reset(nullptr);
-					}
-				});
+			if (VehicleSimulationPT)
+			{
+				VehicleSimulationPT->Terminate();
+				VehicleSimulationPT.Reset(nullptr);
+			}
 		}
 
 	}
diff --git a/Engine/Plugins/Experimental/ChaosModularVehicle/Source/ChaosModularVehicleEngine/Public/ChaosModularVehicle/ChaosSimModuleManagerAsyncCallback.h b/Engine/Plugins/Experimental/ChaosModularVehicle/Source/ChaosModularVehicleEngine/Public/ChaosModularVehicle/ChaosSimModuleManagerAsyncCallback.h
--- a/Engine/Plugins/Experimental/ChaosModularVehicle/Source/ChaosModularVehicleEngine/Public/ChaosModularVehicle/ChaosSimModuleManagerAsyncCallback.h
+++ b/Engine/Plugins/Experimental/ChaosModularVehicle/Source/ChaosModularVehicleEngine/Public/ChaosModularVehicle/ChaosSimModuleManagerAsyncCallback.h
@@ -150,7 +150,10 @@ struct CHAOSMODULARVEHICLEENGINE_API FPhysicsVehicleOutput
 	{
 		for (Chaos::FSimOutputData* Data : SimTreeOutputData)
 		{
-			delete Data;
+			if(Data)
+			{
+				delete Data;
+			}
 		}
 		SimTreeOutputData.Empty();
 	}
