Modular Vehicle Changes/Fixes
* Make reverse a regular control input
* Add hysteresis option to gear change
* Add option to offset wheel force position
* Fix animation bug relating to grabbing output from physics state when there is no lerp to between current and next states

diff --git a/Engine/Plugins/Experimental/ChaosModularVehicle/Source/ChaosModularVehicleEngine/Private/ChaosModularVehicle/ModularVehicleBaseComponent.cpp b/Engine/Plugins/Experimental/ChaosModularVehicle/Source/ChaosModularVehicleEngine/Private/ChaosModularVehicle/ModularVehicleBaseComponent.cpp
--- a/Engine/Plugins/Experimental/ChaosModularVehicle/Source/ChaosModularVehicleEngine/Private/ChaosModularVehicle/ModularVehicleBaseComponent.cpp
+++ b/Engine/Plugins/Experimental/ChaosModularVehicle/Source/ChaosModularVehicleEngine/Private/ChaosModularVehicle/ModularVehicleBaseComponent.cpp
@@ -492,7 +492,7 @@ void UModularVehicleBaseComponent::UpdateState(float DeltaTime)
 		if (!bUsingNetworkPhysicsPrediction)
 		{
 			// and send to server - (ServerUpdateState_Implementation below)
-			ServerUpdateState(InputsContainer, ReverseInput, bKeepVehicleAwake);
+			ServerUpdateState(InputsContainer, bKeepVehicleAwake);
 		}
 
 		if (PawnOwner && PawnOwner->IsNetMode(NM_Client))
@@ -504,20 +504,18 @@ void UModularVehicleBaseComponent::UpdateState(float DeltaTime)
 	{
 		// use replicated values for remote pawns
 		InputsContainer = ReplicatedState.Container;
-		ReverseInput = ReplicatedState.Reverse;
 		bKeepVehicleAwake = ReplicatedState.KeepAwake;
 	}
 }
 
-bool UModularVehicleBaseComponent::ServerUpdateState_Validate(const FModuleInputContainer& InputsIn, bool Reverse, bool KeepAwake)
+bool UModularVehicleBaseComponent::ServerUpdateState_Validate(const FModuleInputContainer& InputsIn, bool KeepAwake)
 {
 	return true;
 }
 
-void UModularVehicleBaseComponent::ServerUpdateState_Implementation(const FModuleInputContainer& InputsIn, bool Reverse, bool KeepAwake)
+void UModularVehicleBaseComponent::ServerUpdateState_Implementation(const FModuleInputContainer& InputsIn, bool KeepAwake)
 {
 	// update state of inputs
-	ReplicatedState.Reverse = Reverse;
 	ReplicatedState.KeepAwake = KeepAwake;
 	ReplicatedState.Container = InputsIn;
 }
@@ -709,7 +707,6 @@ void UModularVehicleBaseComponent::Update(float DeltaTime)
 
 		FModularVehicleAsyncInput* AsyncInput = static_cast<FModularVehicleAsyncInput*>(CurAsyncInput);
 
-		AsyncInput->PhysicsInputs.NetworkInputs.VehicleInputs.Reverse = ReverseInput;
 		AsyncInput->PhysicsInputs.NetworkInputs.VehicleInputs.KeepAwake = bKeepVehicleAwake;
 
 		// All control inputs
@@ -1164,14 +1161,9 @@ void UModularVehicleBaseComponent::SetInput(const FName& Name, const FVector& Va
 	}
 }
 
-void UModularVehicleBaseComponent::SetReverseInput(bool Reverse)
-{
-	RawReverseInput = Reverse;
-}
-
 void UModularVehicleBaseComponent::SetGearInput(int32 Gear)
 {
-	RawGearInput = Gear;
+	GearInput = Gear;
 }
 
 int32 UModularVehicleBaseComponent::GetCurrentGear()
diff --git a/Engine/Plugins/Experimental/ChaosModularVehicle/Source/ChaosModularVehicleEngine/Private/ChaosModularVehicle/ModularVehicleSimulationCU.cpp b/Engine/Plugins/Experimental/ChaosModularVehicle/Source/ChaosModularVehicleEngine/Private/ChaosModularVehicle/ModularVehicleSimulationCU.cpp
--- a/Engine/Plugins/Experimental/ChaosModularVehicle/Source/ChaosModularVehicleEngine/Private/ChaosModularVehicle/ModularVehicleSimulationCU.cpp
+++ b/Engine/Plugins/Experimental/ChaosModularVehicle/Source/ChaosModularVehicleEngine/Private/ChaosModularVehicle/ModularVehicleSimulationCU.cpp
@@ -131,7 +131,6 @@ void FModularVehicleSimulationCU::Simulate_ClusterUnion(UWorld* InWorld, float D
 
 		SimInputData.ControlInputs = &InputInterface;
 		SimInputData.bKeepVehicleAwake = InputData.PhysicsInputs.NetworkInputs.VehicleInputs.KeepAwake;
-		SimInputData.bIsReversing = InputData.PhysicsInputs.NetworkInputs.VehicleInputs.Reverse;
 
 		PerformAdditionalSimWork(InWorld, InputData, Proxy, SimInputData);
 
diff --git a/Engine/Plugins/Experimental/ChaosModularVehicle/Source/ChaosModularVehicleEngine/Private/ChaosModularVehicle/VehicleSimTransmissionComponent.cpp b/Engine/Plugins/Experimental/ChaosModularVehicle/Source/ChaosModularVehicleEngine/Private/ChaosModularVehicle/VehicleSimTransmissionComponent.cpp
--- a/Engine/Plugins/Experimental/ChaosModularVehicle/Source/ChaosModularVehicleEngine/Private/ChaosModularVehicle/VehicleSimTransmissionComponent.cpp
+++ b/Engine/Plugins/Experimental/ChaosModularVehicle/Source/ChaosModularVehicleEngine/Private/ChaosModularVehicle/VehicleSimTransmissionComponent.cpp
@@ -20,6 +20,7 @@ UVehicleSimTransmissionComponent::UVehicleSimTransmissionComponent()
 	ChangeUpRPM = 4500;
 	ChangeDownRPM = 1600;
 	GearChangeTime = 0.5f;
+	GearHysteresisTime = 2.0f;
 	TransmissionEfficiency = 0.9f;
 	TransmissionType = EModuleTransType::Automatic;
 	AutoReverse = true;
@@ -36,6 +37,7 @@ Chaos::ISimulationModuleBase* UVehicleSimTransmissionComponent::CreateNewCoreMod
 	Settings.ChangeUpRPM = ChangeUpRPM;
 	Settings.ChangeDownRPM = ChangeDownRPM;
 	Settings.GearChangeTime = GearChangeTime;
+	Settings.GearHysteresisTime = GearHysteresisTime;
 	Settings.TransmissionEfficiency = TransmissionEfficiency;
 	Settings.TransmissionType = static_cast<Chaos::FTransmissionSettings::ETransType>(TransmissionType);
 
diff --git a/Engine/Plugins/Experimental/ChaosModularVehicle/Source/ChaosModularVehicleEngine/Private/ChaosModularVehicle/VehicleSimWheelComponent.cpp b/Engine/Plugins/Experimental/ChaosModularVehicle/Source/ChaosModularVehicleEngine/Private/ChaosModularVehicle/VehicleSimWheelComponent.cpp
--- a/Engine/Plugins/Experimental/ChaosModularVehicle/Source/ChaosModularVehicleEngine/Private/ChaosModularVehicle/VehicleSimWheelComponent.cpp
+++ b/Engine/Plugins/Experimental/ChaosModularVehicle/Source/ChaosModularVehicleEngine/Private/ChaosModularVehicle/VehicleSimWheelComponent.cpp
@@ -25,6 +25,7 @@ UVehicleSimWheelComponent::UVehicleSimWheelComponent()
 	bTractionControlEnabled = true;
 	AxisType = EWheelAxisType::Y;
 	ReverseDirection = false;
+	ForceOffset = FVector::ZeroVector;
 	bAnimationEnabled = true;
 }
 
@@ -48,6 +49,7 @@ Chaos::ISimulationModuleBase* UVehicleSimWheelComponent::CreateNewCoreModule() c
 	Settings.TractionControlEnabled = bTractionControlEnabled;
 	Settings.Axis = (Chaos::EWheelAxis)(AxisType);
 	Settings.ReverseDirection = ReverseDirection;
+	Settings.ForceOffset = ForceOffset;
 
 	Chaos::ISimulationModuleBase* Wheel = new Chaos::FWheelSimModule(Settings);
 	Wheel->SetAnimationEnabled(bAnimationEnabled);
diff --git a/Engine/Plugins/Experimental/ChaosModularVehicle/Source/ChaosModularVehicleEngine/Public/ChaosModularVehicle/ModularVehicleBaseComponent.h b/Engine/Plugins/Experimental/ChaosModularVehicle/Source/ChaosModularVehicleEngine/Public/ChaosModularVehicle/ModularVehicleBaseComponent.h
--- a/Engine/Plugins/Experimental/ChaosModularVehicle/Source/ChaosModularVehicleEngine/Public/ChaosModularVehicle/ModularVehicleBaseComponent.h
+++ b/Engine/Plugins/Experimental/ChaosModularVehicle/Source/ChaosModularVehicleEngine/Public/ChaosModularVehicle/ModularVehicleBaseComponent.h
@@ -208,9 +208,6 @@ public:
 	UFUNCTION(BlueprintCallable, Category = "Game|Components|ModularVehicle")
 	void SetInputAxis3D(const FName Name, const FVector Value);
 
-	UFUNCTION(BlueprintCallable, Category = "Game|Components|ModularVehicle")
-	void SetReverseInput(bool Reverse);
-
 	/** Set the gear directly */
 	UFUNCTION(BlueprintCallable, Category = "Game|Components|ModularVehicle")
 	void SetGearInput(int32 Gear);
@@ -292,15 +289,7 @@ protected:
 
 	// latest gear selected
 	UPROPERTY(Transient)
-	int32 RawGearInput;
-
-	// reverse direction enbaled
-	UPROPERTY(Transient)
-	bool RawReverseInput;
-
-	// Reverse state
-	UPROPERTY(Transient)
-	bool ReverseInput;
+	int32 GearInput;
 
 	// The currently selected gear
 	UPROPERTY(Transient)
@@ -328,12 +317,9 @@ public:
 	UPROPERTY(Transient, Replicated)
 	TArray<FConstructionData> ConstructionDatas;
 
-	UPROPERTY(transient, duplicatetransient, BlueprintReadOnly, Category = "Game|Components|ModularVehicle")
-	TArray<TObjectPtr<UDefaultModularVehicleInputModifier>> InputModifiers;
-
 	/** Pass current state to server */
 	UFUNCTION(reliable, server, WithValidation)
-	void ServerUpdateState(const FModuleInputContainer& InputsIn, bool Reverse, bool KeepAwake);
+	void ServerUpdateState(const FModuleInputContainer& InputsIn, bool KeepAwake);
 
 	void LogInputSetup();
 
diff --git a/Engine/Plugins/Experimental/ChaosModularVehicle/Source/ChaosModularVehicleEngine/Public/ChaosModularVehicle/VehicleSimTransmissionComponent.h b/Engine/Plugins/Experimental/ChaosModularVehicle/Source/ChaosModularVehicleEngine/Public/ChaosModularVehicle/VehicleSimTransmissionComponent.h
--- a/Engine/Plugins/Experimental/ChaosModularVehicle/Source/ChaosModularVehicleEngine/Public/ChaosModularVehicle/VehicleSimTransmissionComponent.h
+++ b/Engine/Plugins/Experimental/ChaosModularVehicle/Source/ChaosModularVehicleEngine/Public/ChaosModularVehicle/VehicleSimTransmissionComponent.h
@@ -43,6 +43,9 @@ public:
 	UPROPERTY(EditAnywhere, BlueprintReadWrite, Category = Attributes)
 	float GearChangeTime;
 
+	UPROPERTY(EditAnywhere, BlueprintReadWrite, Category = Attributes)
+	float GearHysteresisTime;
+
 	UPROPERTY(EditAnywhere, BlueprintReadWrite, Category = Attributes)
 	float TransmissionEfficiency;
 
diff --git a/Engine/Plugins/Experimental/ChaosModularVehicle/Source/ChaosModularVehicleEngine/Public/ChaosModularVehicle/VehicleSimWheelComponent.h b/Engine/Plugins/Experimental/ChaosModularVehicle/Source/ChaosModularVehicleEngine/Public/ChaosModularVehicle/VehicleSimWheelComponent.h
--- a/Engine/Plugins/Experimental/ChaosModularVehicle/Source/ChaosModularVehicleEngine/Public/ChaosModularVehicle/VehicleSimWheelComponent.h
+++ b/Engine/Plugins/Experimental/ChaosModularVehicle/Source/ChaosModularVehicleEngine/Public/ChaosModularVehicle/VehicleSimWheelComponent.h
@@ -80,6 +80,9 @@ public:
 	UPROPERTY(EditAnywhere, BlueprintReadWrite, Category = Attributes)
 	EWheelAxisType AxisType;
 
+	UPROPERTY(EditAnywhere, BlueprintReadWrite, Category = Attributes)
+	FVector ForceOffset;
+
 	UPROPERTY(EditAnywhere, BlueprintReadWrite, Category = Attributes)
 	bool ReverseDirection;
 
diff --git a/Engine/Source/Programs/HeadlessChaos/Private/HeadlessChaosTestModularVehicles.cpp b/Engine/Source/Programs/HeadlessChaos/Private/HeadlessChaosTestModularVehicles.cpp
--- a/Engine/Source/Programs/HeadlessChaos/Private/HeadlessChaosTestModularVehicles.cpp
+++ b/Engine/Source/Programs/HeadlessChaos/Private/HeadlessChaosTestModularVehicles.cpp
@@ -385,6 +385,7 @@ namespace ChaosTest
 			Setup.ChangeUpRPM = 3000;
 			Setup.ChangeDownRPM = 1200;
 			Setup.GearChangeTime = 0.0f;
+			Setup.GearHysteresisTime = 0.0f;
 			Setup.TransmissionType = FTransmissionSettings::ETransType::AutomaticType;
 			Setup.AutoReverse = false;
 			Setup.TransmissionEfficiency = 1.0f;
diff --git a/Engine/Source/Runtime/Experimental/ChaosVehicles/ChaosVehiclesCore/Private/SimModule/DeferredForcesModular.cpp b/Engine/Source/Runtime/Experimental/ChaosVehicles/ChaosVehiclesCore/Private/SimModule/DeferredForcesModular.cpp
--- a/Engine/Source/Runtime/Experimental/ChaosVehicles/ChaosVehiclesCore/Private/SimModule/DeferredForcesModular.cpp
+++ b/Engine/Source/Runtime/Experimental/ChaosVehicles/ChaosVehiclesCore/Private/SimModule/DeferredForcesModular.cpp
@@ -424,7 +424,7 @@ void FDeferredForcesModular::AddForceAtPosition(Chaos::FPBDRigidParticleHandle*
 	{
 		FTransform FinalOffsetTransform = OffsetTransform;
 		FinalOffsetTransform.AddToTranslation(DataIn.Position);
-		AddForce_Implementation(RigidHandle, FinalOffsetTransform, DataIn.Force, FinalOffsetTransform, (DataIn.Flags & EForceFlags::LevelSlope) == EForceFlags::LevelSlope, DataIn.DebugColor);
+		AddForce_Implementation(RigidHandle, DataIn.OffsetTransform, DataIn.Force, FinalOffsetTransform, (DataIn.Flags & EForceFlags::LevelSlope) == EForceFlags::LevelSlope, DataIn.DebugColor);
 	}
 }
 
diff --git a/Engine/Source/Runtime/Experimental/ChaosVehicles/ChaosVehiclesCore/Private/SimModule/ModuleInput.cpp b/Engine/Source/Runtime/Experimental/ChaosVehicles/ChaosVehiclesCore/Private/SimModule/ModuleInput.cpp
--- a/Engine/Source/Runtime/Experimental/ChaosVehicles/ChaosVehiclesCore/Private/SimModule/ModuleInput.cpp
+++ b/Engine/Source/Runtime/Experimental/ChaosVehicles/ChaosVehiclesCore/Private/SimModule/ModuleInput.cpp
@@ -250,7 +250,7 @@ void FModuleInputContainer::RemoveAllInputs()
 
 void FModuleInputContainer::Lerp(const FModuleInputContainer& Min, const FModuleInputContainer& Max, float Alpha)
 {
-	for (int I = 0; I < InputValues.Num(); I++)
+	for (int I = 0; I < FMath::Min3(InputValues.Num(), Min.InputValues.Num(), Max.InputValues.Num()); I++)
 	{
 		InputValues[I].Lerp(Min.InputValues[I], Max.InputValues[I], Alpha);
 	}
diff --git a/Engine/Source/Runtime/Experimental/ChaosVehicles/ChaosVehiclesCore/Private/SimModule/TransmissionModule.cpp b/Engine/Source/Runtime/Experimental/ChaosVehicles/ChaosVehiclesCore/Private/SimModule/TransmissionModule.cpp
--- a/Engine/Source/Runtime/Experimental/ChaosVehicles/ChaosVehiclesCore/Private/SimModule/TransmissionModule.cpp
+++ b/Engine/Source/Runtime/Experimental/ChaosVehicles/ChaosVehiclesCore/Private/SimModule/TransmissionModule.cpp
@@ -17,6 +17,7 @@ namespace Chaos
 		, TargetGear(1)
 		, CurrentGearChangeTime(0.f)
 		, AllowedToChangeGear(true)
+		, GearHysteresisTimer(0.f)
 	{
 	}
 
@@ -25,7 +26,7 @@ namespace Chaos
 
 		if (Setup().AutoReverse)
 		{
-			if (Inputs.bIsReversing)
+			if (Inputs.GetControls().GetMagnitude(ReverseControlName))
 			{
 				// if reversing change to reverse gear if currently in a forwards gear
 				if (TargetGear > 0)
@@ -45,6 +46,14 @@ namespace Chaos
 
 		if (Setup().TransmissionType == FTransmissionSettings::ETransType::AutomaticType)
 		{
+			if (AllowedToChangeGear == false)
+			{
+				GearHysteresisTimer -= DeltaTime;
+				if (GearHysteresisTimer <= 0.0f)
+				{
+					AllowedToChangeGear = true;
+				}
+			}
 			// not currently changing gear, also don't want to change up because the wheels are spinning up due to having no load
 			if (!IsCurrentlyChangingGear() && AllowedToChangeGear)
 			{
@@ -61,6 +70,8 @@ namespace Chaos
 					if (CurrentGear > 0)
 					{
 						ChangeUp();
+						AllowedToChangeGear = false;
+						GearHysteresisTimer = Setup().GearHysteresisTime;
 					}
 					else
 					{
@@ -90,10 +101,6 @@ namespace Chaos
 			{
 				ChangeDown();
 			}
-			//else if (Inputs.ControlInputs.GearNumber) #TODO: this was never used
-			//{
-			//	TargetGear = Inputs.ControlInputs.GearNumber;
-			//}
 		}
 
 		if (CurrentGear != TargetGear)
diff --git a/Engine/Source/Runtime/Experimental/ChaosVehicles/ChaosVehiclesCore/Private/SimModule/WheelModule.cpp b/Engine/Source/Runtime/Experimental/ChaosVehicles/ChaosVehiclesCore/Private/SimModule/WheelModule.cpp
--- a/Engine/Source/Runtime/Experimental/ChaosVehicles/ChaosVehiclesCore/Private/SimModule/WheelModule.cpp
+++ b/Engine/Source/Runtime/Experimental/ChaosVehicles/ChaosVehiclesCore/Private/SimModule/WheelModule.cpp
@@ -180,8 +180,7 @@ namespace Chaos
 			{
 				ForceFromFriction = -ForceFromFriction;
 			}
-
-			AddLocalForce(SteeringRotator.RotateVector(ForceFromFriction));
+			AddLocalForceAtPosition(SteeringRotator.RotateVector(ForceFromFriction), Setup().ForceOffset);
 		}
 
 		TransmitTorque(VehicleModuleSystem, DriveTorque, BrakeTorque);
diff --git a/Engine/Source/Runtime/Experimental/ChaosVehicles/ChaosVehiclesCore/Public/SimModule/SimulationModuleBase.h b/Engine/Source/Runtime/Experimental/ChaosVehicles/ChaosVehiclesCore/Public/SimModule/SimulationModuleBase.h
--- a/Engine/Source/Runtime/Experimental/ChaosVehicles/ChaosVehiclesCore/Public/SimModule/SimulationModuleBase.h
+++ b/Engine/Source/Runtime/Experimental/ChaosVehicles/ChaosVehiclesCore/Public/SimModule/SimulationModuleBase.h
@@ -36,6 +36,7 @@ namespace Chaos
 	const FName BrakeControlName("Brake");
 	const FName ClutchControlName("Clutch");
 	const FName BoostControlName("Boost");
+	const FName ReverseControlName("Reverse");
 	const FName ChangeUpControlName("ChangeUp");
 	const FName ChangeDownControlName("ChangeDown");
 	const FName PitchControlName("Pitch");
@@ -57,7 +58,6 @@ namespace Chaos
 		FTransform VehicleWorldTransform;
 		TMap<int32, FModuleHitResults> HitResults;
 		FInputInterface* ControlInputs = nullptr;
-		bool bIsReversing = false;
 		bool bKeepVehicleAwake = false;
 	};
 
diff --git a/Engine/Source/Runtime/Experimental/ChaosVehicles/ChaosVehiclesCore/Public/SimModule/TransmissionModule.h b/Engine/Source/Runtime/Experimental/ChaosVehicles/ChaosVehiclesCore/Public/SimModule/TransmissionModule.h
--- a/Engine/Source/Runtime/Experimental/ChaosVehicles/ChaosVehiclesCore/Public/SimModule/TransmissionModule.h
+++ b/Engine/Source/Runtime/Experimental/ChaosVehicles/ChaosVehiclesCore/Public/SimModule/TransmissionModule.h
@@ -73,6 +73,7 @@ namespace Chaos
 			, ChangeUpRPM(5000)
 			, ChangeDownRPM(2500)
 			, GearChangeTime(0.5f)
+			, GearHysteresisTime(2.0f)
 			, TransmissionEfficiency(1.f)
 			, TransmissionType(ETransType::AutomaticType)
 			, AutoReverse(true)
@@ -92,6 +93,7 @@ namespace Chaos
 		uint32 ChangeUpRPM;				// [RPM]
 		uint32 ChangeDownRPM;			// [RPM]
 		float GearChangeTime; 			// [sec]
+		float GearHysteresisTime;		// [sec]
 
 		float TransmissionEfficiency;	// Loss from friction in the system mean we might run at around 0.94 Efficiency
 
@@ -174,7 +176,7 @@ namespace Chaos
 		float CurrentGearChangeTime; // Time to change gear, no power transmitted to the wheels during change
 
 		bool AllowedToChangeGear; // conditions are ok for an automatic gear change
-
+		float GearHysteresisTimer;
 	};
 	
 	class CHAOSVEHICLESCORE_API FTransmissionSimFactory
diff --git a/Engine/Source/Runtime/Experimental/ChaosVehicles/ChaosVehiclesCore/Public/SimModule/WheelModule.h b/Engine/Source/Runtime/Experimental/ChaosVehicles/ChaosVehiclesCore/Public/SimModule/WheelModule.h
--- a/Engine/Source/Runtime/Experimental/ChaosVehicles/ChaosVehiclesCore/Public/SimModule/WheelModule.h
+++ b/Engine/Source/Runtime/Experimental/ChaosVehicles/ChaosVehiclesCore/Public/SimModule/WheelModule.h
@@ -85,8 +85,8 @@ namespace Chaos
 			, MaxRotationVel(100.0f)
 			, Axis(EWheelAxis::X)
 			, ReverseDirection(false)
+			, ForceOffset(FVector::ZeroVector)
 		{
-
 		}
 
 		float Radius;
@@ -114,6 +114,8 @@ namespace Chaos
 		float MaxRotationVel;
 		EWheelAxis Axis;
 		bool ReverseDirection;
+		FVector ForceOffset;
+
 	};
 
 	class CHAOSVEHICLESCORE_API FWheelSimModule : public FWheelBaseInterface, public TSimModuleSettings<FWheelSettings>, public TSimulationModuleTypeable<FWheelSimModule>
