Remove unused line

diff --git a/Engine/Plugins/Experimental/ChaosModularVehicle/Source/ChaosModularVehicleEngine/Private/ChaosModularVehicle/ModularVehicleBaseComponent.cpp b/Engine/Plugins/Experimental/ChaosModularVehicle/Source/ChaosModularVehicleEngine/Private/ChaosModularVehicle/ModularVehicleBaseComponent.cpp
--- a/Engine/Plugins/Experimental/ChaosModularVehicle/Source/ChaosModularVehicleEngine/Private/ChaosModularVehicle/ModularVehicleBaseComponent.cpp
+++ b/Engine/Plugins/Experimental/ChaosModularVehicle/Source/ChaosModularVehicleEngine/Private/ChaosModularVehicle/ModularVehicleBaseComponent.cpp
@@ -650,7 +650,6 @@ void UModularVehicleBaseComponent::ParallelUpdate(float DeltaTime)
 					else
 					{
 						Chaos::FSimOutputData* CurrentSimData = CurrentOutput->VehicleSimOutput.SimTreeOutputData[I];
-						Chaos::FSimOutputData* NextSimData = NextOutput->VehicleSimOutput.SimTreeOutputData[I];
 						PVehicleOutput->SimTreeOutputData.EmplaceAt(I, CurrentSimData->MakeNewData());
 						*PVehicleOutput->SimTreeOutputData[I] = *CurrentSimData;
 					}
