Modular Vehicle: Remove obsolete debug logic

diff --git a/Engine/Plugins/Experimental/ChaosModularVehicle/Source/ChaosModularVehicleEngine/Private/ChaosModularVehicle/ModularVehicleSimulationCU.cpp b/Engine/Plugins/Experimental/ChaosModularVehicle/Source/ChaosModularVehicleEngine/Private/ChaosModularVehicle/ModularVehicleSimulationCU.cpp
--- a/Engine/Plugins/Experimental/ChaosModularVehicle/Source/ChaosModularVehicleEngine/Private/ChaosModularVehicle/ModularVehicleSimulationCU.cpp
+++ b/Engine/Plugins/Experimental/ChaosModularVehicle/Source/ChaosModularVehicleEngine/Private/ChaosModularVehicle/ModularVehicleSimulationCU.cpp
@@ -39,32 +39,6 @@ void FModularVehicleSimulationCU::Terminate()
 
 void FModularVehicleSimulationCU::Simulate(UWorld* InWorld, float DeltaSeconds, const FModularVehicleAsyncInput& InputData, FModularVehicleAsyncOutput& OutputData, IPhysicsProxyBase* Proxy)
 {
-#if DEBUG_NETWORK_PHYSICS
-
-	if (InWorld->IsNetMode(NM_ListenServer) || InWorld->IsNetMode(NM_DedicatedServer))
-	{
-		UE_LOG(LogTemp, Log, TEXT("SERVER | PT | TickVehicle | Async tick vehicle with inputs at frame %d : Throttle = %f Brake = %f Roll = %f Pitch = %f Yaw = %f Steering = %f Handbrake = %f"),
-			InputData.PhysicsInputs.NetworkInputs.LocalFrame, VehicleInputs.Throttle, VehicleInputs.Brake, VehicleInputs.Roll, VehicleInputs.Pitch,
-			VehicleInputs.Yaw, VehicleInputs.Steering, VehicleInputs.Handbrake);
-
-		UE_LOG(LogTemp, Log, TEXT("ALT-SERVER | PT | TickVehicle | Async tick vehicle with inputs at frame %d : Throttle = %f Brake = %f Roll = %f Pitch = %f Yaw = %f Steering = %f Handbrake = %f"),
-			InputData.PhysicsInputs.NetworkInputs.LocalFrame, InputData.PhysicsInputs.NetworkInputs.VehicleInputs.Throttle, InputData.PhysicsInputs.NetworkInputs.VehicleInputs.Brake, InputData.PhysicsInputs.NetworkInputs.VehicleInputs.Roll, InputData.PhysicsInputs.NetworkInputs.VehicleInputs.Pitch,
-			InputData.PhysicsInputs.NetworkInputs.VehicleInputs.Yaw, InputData.PhysicsInputs.NetworkInputs.VehicleInputs.Steering, InputData.PhysicsInputs.NetworkInputs.VehicleInputs.Handbrake);
-	}
-	else if (InWorld->IsNetMode(NM_Client))
-	{
-		UE_LOG(LogTemp, Log, TEXT("CLIENT | PT | TickVehicle | Async tick vehicle with inputs at frame %d : Throttle = %f Brake = %f Roll = %f Pitch = %f Yaw = %f Steering = %f Handbrake = %f"),
-			InputData.PhysicsInputs.NetworkInputs.LocalFrame, VehicleInputs.Throttle, VehicleInputs.Brake, VehicleInputs.Roll, VehicleInputs.Pitch,
-			VehicleInputs.Yaw, VehicleInputs.Steering, VehicleInputs.Handbrake);
-
-		UE_LOG(LogTemp, Log, TEXT("ALT-CLIENT | PT | TickVehicle | Async tick vehicle with inputs at frame %d : Throttle = %f Brake = %f Roll = %f Pitch = %f Yaw = %f Steering = %f Handbrake = %f"),
-			InputData.PhysicsInputs.NetworkInputs.LocalFrame, InputData.PhysicsInputs.NetworkInputs.VehicleInputs.Throttle, InputData.PhysicsInputs.NetworkInputs.VehicleInputs.Brake, InputData.PhysicsInputs.NetworkInputs.VehicleInputs.Roll, InputData.PhysicsInputs.NetworkInputs.VehicleInputs.Pitch,
-			InputData.PhysicsInputs.NetworkInputs.VehicleInputs.Yaw, InputData.PhysicsInputs.NetworkInputs.VehicleInputs.Steering, InputData.PhysicsInputs.NetworkInputs.VehicleInputs.Handbrake);
-	}
-
-#endif
-
-
 	Chaos::FPBDRigidsSolver* RigidsSolver = static_cast<Chaos::FPBDRigidsSolver*>(Proxy->GetSolver<Chaos::FPhysicsSolver>());
 	int CurrentFrame = -1;
 	if (RigidsSolver != nullptr)
