Modular vehicle: Fix vehicle output values, certain conditions didn't lerp between physics frames

diff --git a/Engine/Plugins/Experimental/ChaosModularVehicle/Source/ChaosModularVehicleEngine/Private/ChaosModularVehicle/ChaosSimModuleManager.cpp b/Engine/Plugins/Experimental/ChaosModularVehicle/Source/ChaosModularVehicleEngine/Private/ChaosModularVehicle/ChaosSimModuleManager.cpp
--- a/Engine/Plugins/Experimental/ChaosModularVehicle/Source/ChaosModularVehicleEngine/Private/ChaosModularVehicle/ChaosSimModuleManager.cpp
+++ b/Engine/Plugins/Experimental/ChaosModularVehicle/Source/ChaosModularVehicleEngine/Private/ChaosModularVehicle/ChaosSimModuleManager.cpp
@@ -238,7 +238,7 @@ void FChaosSimModuleManager::ParallelUpdateVehicles(float DeltaSeconds)
 	}
 
 	// Since we are in pre-physics, delta seconds is not accounted for in external time yet
-	const float ResultsTime = AsyncCallback->GetSolver()->GetPhysicsResultsTime_External() + DeltaSeconds;
+	const float ResultsTime = AsyncCallback->GetSolver()->GetPhysicsResultsTime_External();
 
 	// Find index of first non-consumable output (first one after current time)
 	int32 LastOutputIdx = 0;
diff --git a/Engine/Plugins/Experimental/ChaosModularVehicle/Source/ChaosModularVehicleEngine/Private/ChaosModularVehicle/ModularVehicleBaseComponent.cpp b/Engine/Plugins/Experimental/ChaosModularVehicle/Source/ChaosModularVehicleEngine/Private/ChaosModularVehicle/ModularVehicleBaseComponent.cpp
--- a/Engine/Plugins/Experimental/ChaosModularVehicle/Source/ChaosModularVehicleEngine/Private/ChaosModularVehicle/ModularVehicleBaseComponent.cpp
+++ b/Engine/Plugins/Experimental/ChaosModularVehicle/Source/ChaosModularVehicleEngine/Private/ChaosModularVehicle/ModularVehicleBaseComponent.cpp
@@ -639,13 +639,6 @@ void UModularVehicleBaseComponent::ParallelUpdate(float DeltaTime)
 						Chaos::FSimOutputData* NextSimData = NextOutput->VehicleSimOutput.SimTreeOutputData[I];
 						PVehicleOutput->SimTreeOutputData.EmplaceAt(I, CurrentSimData->MakeNewData());
 						PVehicleOutput->SimTreeOutputData[I]->Lerp(*CurrentSimData, *NextSimData, OutputInterpAlpha);
-
-#if !(UE_BUILD_SHIPPING || UE_BUILD_TEST)
-						if (PVehicleOutput && !PVehicleOutput->SimTreeOutputData.IsEmpty() && PVehicleOutput->SimTreeOutputData[I])
-						{
-							PVehicleOutput->SimTreeOutputData[I]->DebugString = CurrentOutput->VehicleSimOutput.SimTreeOutputData[I]->DebugString;
-						}
-#endif
 					}
 					else
 					{
@@ -653,43 +646,65 @@ void UModularVehicleBaseComponent::ParallelUpdate(float DeltaTime)
 						PVehicleOutput->SimTreeOutputData.EmplaceAt(I, CurrentSimData->MakeNewData());
 						*PVehicleOutput->SimTreeOutputData[I] = *CurrentSimData;
 					}
+				}
+			}
+			else
+			{
+				for (int I = 0; I < NumItems; I++)
+				{
+					Chaos::FSimOutputData* CurrentSimData = CurrentOutput->VehicleSimOutput.SimTreeOutputData[I];
+					PVehicleOutput->SimTreeOutputData.EmplaceAt(I, CurrentSimData->MakeNewData());
+					*PVehicleOutput->SimTreeOutputData[I] = *CurrentSimData;
+				}
+			}
 
-					// extract/cache some generally useful values as we go as trying to locate this data later requires a search
-					if (PVehicleOutput->SimTreeOutputData[I]->IsSimType<Chaos::FTransmissionSimModule>())
-					{
-						// if there is more than one transmission then the last one will inform us of the current gear
-						CurrentGear = static_cast<Chaos::FTransmissionOutputData*>(PVehicleOutput->SimTreeOutputData[I])->CurrentGear;
-					}
-					else if (PVehicleOutput->SimTreeOutputData[I]->IsSimType<Chaos::FEngineSimModule>())
-					{
-						// if there is more than one engine then the last one will inform us of the engine RPM
-						Chaos::FEngineOutputData* Engine = static_cast<Chaos::FEngineOutputData*>(CurrentOutput->VehicleSimOutput.SimTreeOutputData[I]);
 
-						EngineRPM = Engine->RPM;
-						EngineTorque = Engine->Torque;
-					}
+			for (int I = 0; I < NumItems; I++)
+			{
+				// extract/cache some generally useful values as we go as trying to locate this data later requires a search
+				if (PVehicleOutput->SimTreeOutputData[I]->IsSimType<Chaos::FTransmissionSimModule>())
+				{
+					// if there is more than one transmission then the last one will inform us of the current gear
+					CurrentGear = static_cast<Chaos::FTransmissionOutputData*>(PVehicleOutput->SimTreeOutputData[I])->CurrentGear;
+				}
+				else if (PVehicleOutput->SimTreeOutputData[I]->IsSimType<Chaos::FEngineSimModule>())
+				{
+					// if there is more than one engine then the last one will inform us of the engine RPM
+					Chaos::FEngineOutputData* Engine = static_cast<Chaos::FEngineOutputData*>(CurrentOutput->VehicleSimOutput.SimTreeOutputData[I]);
+
+					EngineRPM = Engine->RPM;
+					EngineTorque = Engine->Torque;
+				}
 
-					if (Chaos::FSimOutputData* ModuleOutput = PVehicleOutput->SimTreeOutputData[I])
+
+				if (Chaos::FSimOutputData* ModuleOutput = PVehicleOutput->SimTreeOutputData[I])
+				{
+					if ((ModuleOutput->AnimationSetupIndex >= 0) && (ModuleOutput->AnimationSetupIndex < ModuleAnimationSetups.Num()))
 					{
-						if ((ModuleOutput->AnimationSetupIndex >= 0) && (ModuleOutput->AnimationSetupIndex < ModuleAnimationSetups.Num()))
-						{
-							ModuleAnimationSetups[ModuleOutput->AnimationSetupIndex].AnimFlags |= ModuleOutput->AnimFlags;
+						ModuleAnimationSetups[ModuleOutput->AnimationSetupIndex].AnimFlags |= ModuleOutput->AnimFlags;
 
-							if (ModuleOutput->AnimFlags & Chaos::EAnimationFlags::AnimateRotation)
-							{
-								ModuleAnimationSetups[ModuleOutput->AnimationSetupIndex].RotOffset = ModuleOutput->AnimationRotOffset;
-							}
+						if (ModuleOutput->AnimFlags & Chaos::EAnimationFlags::AnimateRotation)
+						{
+							ModuleAnimationSetups[ModuleOutput->AnimationSetupIndex].RotOffset = ModuleOutput->AnimationRotOffset;
+						}
 
-							if (ModuleOutput->AnimFlags & Chaos::EAnimationFlags::AnimatePosition)
-							{
-								ModuleAnimationSetups[ModuleOutput->AnimationSetupIndex].LocOffset = ModuleOutput->AnimationLocOffset;
-							}
+						if (ModuleOutput->AnimFlags & Chaos::EAnimationFlags::AnimatePosition)
+						{
+							ModuleAnimationSetups[ModuleOutput->AnimationSetupIndex].LocOffset = ModuleOutput->AnimationLocOffset;
 						}
 					}
+				}
 
+#if !(UE_BUILD_SHIPPING || UE_BUILD_TEST)
+				if (PVehicleOutput && !PVehicleOutput->SimTreeOutputData.IsEmpty() && PVehicleOutput->SimTreeOutputData[I])
+				{
+					PVehicleOutput->SimTreeOutputData[I]->DebugString = CurrentOutput->VehicleSimOutput.SimTreeOutputData[I]->DebugString;
 				}
+#endif
+
 			}
 
+
 		}
 	}
 }
