[NetPhysics] Network optimization for modular vehicle
- Only send inputs that have a non zero value
- API to only send modules that are not in their default state

diff --git a/Engine/Plugins/Experimental/ChaosModularVehicle/Source/ChaosModularVehicleEngine/Private/ChaosModularVehicle/ChaosSimModuleManagerAsyncCallback.cpp b/Engine/Plugins/Experimental/ChaosModularVehicle/Source/ChaosModularVehicleEngine/Private/ChaosModularVehicle/ChaosSimModuleManagerAsyncCallback.cpp
--- a/Engine/Plugins/Experimental/ChaosModularVehicle/Source/ChaosModularVehicleEngine/Private/ChaosModularVehicle/ChaosSimModuleManagerAsyncCallback.cpp
+++ b/Engine/Plugins/Experimental/ChaosModularVehicle/Source/ChaosModularVehicleEngine/Private/ChaosModularVehicle/ChaosSimModuleManagerAsyncCallback.cpp
@@ -9,6 +9,7 @@
 #include "PhysicsProxy/GeometryCollectionPhysicsProxy.h"
 #include "PhysicsProxy/ClusterUnionPhysicsProxy.h"
 #include "SimModule/ModuleFactoryRegister.h"
+#include "Net/Core/NetBitArray.h"
 
 FSimModuleDebugParams GSimModuleDebugParams;
 
@@ -305,48 +306,90 @@ void FNetworkModularVehicleInputs::MergeData(const FNetworkPhysicsData& FromData
 
 bool FNetworkModularVehicleStates::NetSerialize(FArchive& Ar, class UPackageMap* Map, bool& bOutSuccess)
 {
+	using namespace UE::Net;
+
 	FNetworkPhysicsData::SerializeFrames(Ar);
 
-	int32 NumNetModules = ModuleData.Num();
-	Ar << NumNetModules;
-	if(Ar.IsLoading() && NumNetModules != ModuleData.Num())
+	uint32 NumNetModules = ModuleData.Num();
+	Ar.SerializeIntPacked(NumNetModules);
+
+	// Array of bits to mark which modules to serialize or not
+	FNetBitArray ModulesBitArray(NumNetModules);
+
+	if (Ar.IsLoading() && NumNetModules != ModuleData.Num())
 	{
 		ModuleData.Reserve(NumNetModules);
 	}
-	for (int I = 0; I < NumNetModules; I++)
+
+	if (Ar.IsLoading())
 	{
-		if (Ar.IsLoading())
+		Ar.SerializeBits(ModulesBitArray.GetData(), NumNetModules);
+
+		for (uint32 I = 0; I < NumNetModules; I++)
 		{
 			if (NumNetModules > 0)
 			{
 				uint32 ModuleTypeHash = 0;
-				int32 SimArrayIndex = 0;
+				uint32 SimArrayIndexUnsigned = 0;
+
 				Ar << ModuleTypeHash;
-				Ar << SimArrayIndex;
+				Ar.SerializeIntPacked(SimArrayIndexUnsigned);
+
+				const int32 SimArrayIndex = static_cast<int32>(SimArrayIndexUnsigned) - 1; // Convert back to signed and adjust
 
-				if (I >= ModuleData.Num())
+				if (I >= static_cast<uint32>(ModuleData.Num()))
 				{
 					if (TSharedPtr<Chaos::FModuleNetData> Data = Chaos::FModuleFactoryRegister::Get().GenerateNetData(ModuleTypeHash, SimArrayIndex))
 					{
 						ModuleData.Emplace(Data);
 					}
 				}
-				if(I <= ModuleData.Num() && ModuleData[I].IsValid())
+
+				if (I <= static_cast<uint32>(ModuleData.Num()) && ModuleData[I].IsValid())
 				{
 					check(ModuleTypeHash == Chaos::FModuleFactoryRegister::GetModuleHash(ModuleData[I]->GetSimType()));
-					ModuleData[I]->Serialize(Ar);
+
+					const bool bHasSerializedData = ModulesBitArray.IsBitSet(I);
+					if (bHasSerializedData)
+					{
+						ModuleData[I]->Serialize(Ar);
+					}
+					else
+					{
+						ModuleData[I]->ApplyDefaultState();
+					}
 				}
 			}
 		}
-		else
+	}
+	else
+	{
+		// Only mark modules for serialization if they are not in their default state
+		for (uint32 I = 0; I < NumNetModules; I++)
 		{
-			int ModuleTypeHash = Chaos::FModuleFactoryRegister::GetModuleHash(ModuleData[I]->GetSimType());
-			
-			Ar << ModuleTypeHash;
-			Ar << ModuleData[I]->SimArrayIndex;
-			ModuleData[I]->Serialize(Ar);
+			if (ModuleData[I]->IsDefaultState() == false)
+			{
+				ModulesBitArray.SetBit(I);
+			}
 		}
+		Ar.SerializeBits(ModulesBitArray.GetData(), NumNetModules);
+
+		for (uint32 I = 0; I < NumNetModules; I++)
+		{
+			uint32 ModuleTypeHash = Chaos::FModuleFactoryRegister::GetModuleHash(ModuleData[I]->GetSimType());
 
+			check(ModuleData[I]->SimArrayIndex + 1 >= 0);
+			uint32 SimArrayIndexUnsigned = static_cast<uint32>(ModuleData[I]->SimArrayIndex + 1); // Convert to unsigned and align default -1 to 0. Done to be able to use SerializeIntPacked() for network optimization. 
+
+			Ar << ModuleTypeHash;
+			Ar.SerializeIntPacked(SimArrayIndexUnsigned);
+
+			const bool bShouldSerializeData = ModulesBitArray.IsBitSet(I);
+			if (bShouldSerializeData)
+			{
+				ModuleData[I]->Serialize(Ar);
+			}
+		}
 	}
 
 	return true;
diff --git a/Engine/Source/Runtime/Experimental/ChaosVehicles/ChaosVehiclesCore/Private/SimModule/ModuleInput.cpp b/Engine/Source/Runtime/Experimental/ChaosVehicles/ChaosVehiclesCore/Private/SimModule/ModuleInput.cpp
--- a/Engine/Source/Runtime/Experimental/ChaosVehicles/ChaosVehiclesCore/Private/SimModule/ModuleInput.cpp
+++ b/Engine/Source/Runtime/Experimental/ChaosVehicles/ChaosVehiclesCore/Private/SimModule/ModuleInput.cpp
@@ -1,6 +1,7 @@
 // Copyright Epic Games, Inc. All Rights Reserved.
 
 #include "SimModule/ModuleInput.h"
+#include "Net/Core/NetBitArray.h"
 
 TArray<struct FModuleInputSetup>* FScopedModuleInputInitializer::InitSetupData = nullptr;
 
@@ -217,25 +218,52 @@ void FModuleInputContainer::ZeroValues()
 
 void FModuleInputContainer::Serialize(FArchive& Ar, UPackageMap* Map, bool& bOutSuccess)
 {
+	using namespace UE::Net;
+
 	bOutSuccess = true;
 
 	uint32 Number = InputValues.Num();
-	Ar << Number;
+	Ar.SerializeIntPacked(Number);
+
+	// Array of bits to mark which inputs to serialize or not
+	FNetBitArray InputsBitArray(Number);
+
 	if (Ar.IsLoading())
 	{
+		Ar.SerializeBits(InputsBitArray.GetData(), Number);
+
 		InputValues.Reset(Number);
 		for (uint32 I = 0; I < Number; I++)
 		{
 			FModuleInputValue Value;
-			Value.Serialize(Ar, Map, bOutSuccess);
+
+			const bool bHasNonZeroValue = InputsBitArray.IsBitSet(I);
+			if (bHasNonZeroValue)
+			{
+				Value.Serialize(Ar, Map, bOutSuccess);
+			}
 			InputValues.EmplaceAt(I, Value);
 		}
 	}
 	else
 	{
+		// Only mark inputs for serialization if they have a non-zero value
 		for (uint32 I = 0; I < Number; I++)
 		{
-			InputValues[I].Serialize(Ar, Map, bOutSuccess);
+			if (InputValues[I].IsNonZero())
+			{
+				InputsBitArray.SetBit(I);
+			}
+		}
+		Ar.SerializeBits(InputsBitArray.GetData(), Number);
+
+		for (uint32 I = 0; I < Number; I++)
+		{
+			const bool bHasNonZeroValue = InputsBitArray.IsBitSet(I);
+			if (bHasNonZeroValue)
+			{
+				InputValues[I].Serialize(Ar, Map, bOutSuccess);
+			}
 		}
 	}
 
diff --git a/Engine/Source/Runtime/Experimental/ChaosVehicles/ChaosVehiclesCore/Public/SimModule/SimulationModuleBase.h b/Engine/Source/Runtime/Experimental/ChaosVehicles/ChaosVehiclesCore/Public/SimModule/SimulationModuleBase.h
--- a/Engine/Source/Runtime/Experimental/ChaosVehicles/ChaosVehiclesCore/Public/SimModule/SimulationModuleBase.h
+++ b/Engine/Source/Runtime/Experimental/ChaosVehicles/ChaosVehiclesCore/Public/SimModule/SimulationModuleBase.h
@@ -472,7 +472,7 @@ namespace Chaos
 	/**
 	* Interface base class for all module network serialization
 	*/
-	struct CHAOSVEHICLESCORE_API FModuleNetData :public FSimulationModuleTypeContainer, public TSimulationModuleTypeableBase<FModuleNetData>
+	struct CHAOSVEHICLESCORE_API FModuleNetData : public FSimulationModuleTypeContainer, public TSimulationModuleTypeableBase<FModuleNetData>
 	{
 		DEFINE_CHAOSSIMTYPENAME(FModuleNetData);
 		FModuleNetData(int InSimArrayIndex, const FString& InDebugString = FString())
@@ -488,6 +488,15 @@ namespace Chaos
 		virtual void FillNetState(const ISimulationModuleBase* SimModule) = 0;
 		virtual void FillSimState(ISimulationModuleBase* SimModule) = 0;
 		virtual void Lerp(const float LerpFactor, const FModuleNetData& Max, const FModuleNetData& MaxValue) = 0;
+		
+		/** Return true if the current state of this module is the default state,
+		* If the state is default it will not be sent over the network, instead ApplyDefaultState() will be called on the receiving end to apply the default state for this module.
+		* Important: Be sure to also implement ApplyDefaultState() function to ensure that the data gets set to their desired values when not serialized.
+		* Example: bools that are false or true by default and floats that are 0.0f doesn't need to be networked */
+		virtual bool IsDefaultState() { return false; }
+		
+		/** Set the states data to default values, called during deserialization if IsDefaultState() was true during serialization */
+		virtual void ApplyDefaultState() { }
 
 		int SimArrayIndex = -1;
 
