Add CollisionChannel parameter for suspension traces

#rnx

diff --git a/Engine/Plugins/Experimental/ChaosModularVehicle/Source/ChaosModularVehicleEngine/Private/ChaosModularVehicle/ModularVehicleBaseComponent.cpp b/Engine/Plugins/Experimental/ChaosModularVehicle/Source/ChaosModularVehicleEngine/Private/ChaosModularVehicle/ModularVehicleBaseComponent.cpp
--- a/Engine/Plugins/Experimental/ChaosModularVehicle/Source/ChaosModularVehicleEngine/Private/ChaosModularVehicle/ModularVehicleBaseComponent.cpp
+++ b/Engine/Plugins/Experimental/ChaosModularVehicle/Source/ChaosModularVehicleEngine/Private/ChaosModularVehicle/ModularVehicleBaseComponent.cpp
@@ -50,6 +50,7 @@ UModularVehicleBaseComponent::UModularVehicleBaseComponent(const FObjectInitiali
 
 	SuspensionTraceCollisionResponses = FCollisionResponseContainer::GetDefaultResponseContainer();
 	SuspensionTraceCollisionResponses.Vehicle = ECR_Ignore;
+	SuspensionCollisionChannel = ECollisionChannel::ECC_WorldDynamic;
 	bSuspensionTraceComplex = true;
 	bKeepVehicleAwake = true;
 
@@ -730,6 +731,7 @@ void UModularVehicleBaseComponent::Update(float DeltaTime)
 		TraceParams.bReturnPhysicalMaterial = true;	// we need this to get the surface friction coefficient
 		TraceParams.AddIgnoredActors(ActorsToIgnore);
 		TraceParams.bTraceComplex = bSuspensionTraceComplex;
+		AsyncInput->PhysicsInputs.CollisionChannel = SuspensionCollisionChannel;
 		AsyncInput->PhysicsInputs.TraceParams = TraceParams;
 		AsyncInput->PhysicsInputs.TraceCollisionResponse = SuspensionTraceCollisionResponses;
 		AsyncInput->PhysicsInputs.TraceType = TraceType;
diff --git a/Engine/Plugins/Experimental/ChaosModularVehicle/Source/ChaosModularVehicleEngine/Private/ChaosModularVehicle/ModularVehicleSimulationCU.cpp b/Engine/Plugins/Experimental/ChaosModularVehicle/Source/ChaosModularVehicleEngine/Private/ChaosModularVehicle/ModularVehicleSimulationCU.cpp
--- a/Engine/Plugins/Experimental/ChaosModularVehicle/Source/ChaosModularVehicleEngine/Private/ChaosModularVehicle/ModularVehicleSimulationCU.cpp
+++ b/Engine/Plugins/Experimental/ChaosModularVehicle/Source/ChaosModularVehicleEngine/Private/ChaosModularVehicle/ModularVehicleSimulationCU.cpp
@@ -191,7 +191,7 @@ void FModularVehicleSimulationCU::PerformAdditionalSimWork(UWorld* InWorld, cons
 								FVector TraceNormal = TraceVector.GetSafeNormal();
 
 								FHitResult HitResult = FHitResult();
-								ECollisionChannel SpringCollisionChannel = ECollisionChannel::ECC_WorldDynamic;
+								ECollisionChannel SpringCollisionChannel = InputData.PhysicsInputs.CollisionChannel;
 								const FCollisionResponseParams& ResponseParams = InputData.PhysicsInputs.TraceCollisionResponse;
 								if (InWorld)
 								{
diff --git a/Engine/Plugins/Experimental/ChaosModularVehicle/Source/ChaosModularVehicleEngine/Public/ChaosModularVehicle/ChaosSimModuleManagerAsyncCallback.h b/Engine/Plugins/Experimental/ChaosModularVehicle/Source/ChaosModularVehicleEngine/Public/ChaosModularVehicle/ChaosSimModuleManagerAsyncCallback.h
--- a/Engine/Plugins/Experimental/ChaosModularVehicle/Source/ChaosModularVehicleEngine/Public/ChaosModularVehicle/ChaosSimModuleManagerAsyncCallback.h
+++ b/Engine/Plugins/Experimental/ChaosModularVehicle/Source/ChaosModularVehicleEngine/Public/ChaosModularVehicle/ChaosSimModuleManagerAsyncCallback.h
@@ -190,12 +190,14 @@ enum class ETraceType : uint8
 struct CHAOSMODULARVEHICLEENGINE_API FPhysicsModularVehicleInputs
 {
 	FPhysicsModularVehicleInputs()
-		: TraceParams()
+		: CollisionChannel(ECollisionChannel::ECC_WorldDynamic)
+		, TraceParams()
 		, TraceCollisionResponse()
 		, TraceType(ETraceType::Raycast)
 	{
 	}
 	mutable FNetworkModularVehicleInputs NetworkInputs;
+	mutable ECollisionChannel CollisionChannel;
 	mutable FCollisionQueryParams TraceParams;
 	mutable FCollisionResponseContainer TraceCollisionResponse;
 	mutable ETraceType TraceType;
diff --git a/Engine/Plugins/Experimental/ChaosModularVehicle/Source/ChaosModularVehicleEngine/Public/ChaosModularVehicle/ModularVehicleBaseComponent.h b/Engine/Plugins/Experimental/ChaosModularVehicle/Source/ChaosModularVehicleEngine/Public/ChaosModularVehicle/ModularVehicleBaseComponent.h
--- a/Engine/Plugins/Experimental/ChaosModularVehicle/Source/ChaosModularVehicleEngine/Public/ChaosModularVehicle/ModularVehicleBaseComponent.h
+++ b/Engine/Plugins/Experimental/ChaosModularVehicle/Source/ChaosModularVehicleEngine/Public/ChaosModularVehicle/ModularVehicleBaseComponent.h
@@ -158,6 +158,10 @@ public:
 	UPROPERTY(EditAnywhere, Category = "Game|Components|ModularVehicle")
 	struct FCollisionResponseContainer SuspensionTraceCollisionResponses;
 
+	/** Collision channel to use for the suspension trace. */
+	UPROPERTY(EditAnywhere, Category = "Game|Components|ModularVehicle")
+	TEnumAsByte<ECollisionChannel> SuspensionCollisionChannel;
+
 	UPROPERTY(EditAnywhere, Category = "Game|Components|ModularVehicle")
 	bool bSuspensionTraceComplex;
 
