Move to transactionally safe RWLocks for ModularVehicleSimulationCU

diff --git a/Engine/Plugins/Experimental/ChaosModularVehicle/Source/ChaosModularVehicleEngine/Private/ChaosModularVehicle/ModularVehicleSimulationCU.cpp b/Engine/Plugins/Experimental/ChaosModularVehicle/Source/ChaosModularVehicleEngine/Private/ChaosModularVehicle/ModularVehicleSimulationCU.cpp
--- a/Engine/Plugins/Experimental/ChaosModularVehicle/Source/ChaosModularVehicleEngine/Private/ChaosModularVehicle/ModularVehicleSimulationCU.cpp
+++ b/Engine/Plugins/Experimental/ChaosModularVehicle/Source/ChaosModularVehicleEngine/Private/ChaosModularVehicle/ModularVehicleSimulationCU.cpp
@@ -98,7 +98,7 @@ void FModularVehicleSimulationCU::Simulate_ClusterUnion(UWorld* InWorld, float D
 		//		, *Proxy->GetParticle_Internal()->W().ToString()));
 		//}
 
-		FReadScopeLock InputConfigLock(InputConfigurationLock);
+		FTransactionallySafeReadScopeLock InputConfigLock(InputConfigurationLock);
 
 		FModuleInputContainer Container = InputData.PhysicsInputs.NetworkInputs.VehicleInputs.Container;
 		FInputInterface InputInterface(InputNameMap, Container);
diff --git a/Engine/Plugins/Experimental/ChaosModularVehicle/Source/ChaosModularVehicleEngine/Public/ChaosModularVehicle/ModularVehicleSimulationCU.h b/Engine/Plugins/Experimental/ChaosModularVehicle/Source/ChaosModularVehicleEngine/Public/ChaosModularVehicle/ModularVehicleSimulationCU.h
--- a/Engine/Plugins/Experimental/ChaosModularVehicle/Source/ChaosModularVehicleEngine/Public/ChaosModularVehicle/ModularVehicleSimulationCU.h
+++ b/Engine/Plugins/Experimental/ChaosModularVehicle/Source/ChaosModularVehicleEngine/Public/ChaosModularVehicle/ModularVehicleSimulationCU.h
@@ -7,6 +7,7 @@
 #include "SimModule/SimModuleTree.h"
 #include "SimModule/SimModulesInclude.h"
 #include "ModularVehicleBuilder.h"
+#include "Misc/TransactionallySafeRWLock.h"
 
 struct FModularVehicleAsyncInput;
 struct FChaosSimModuleManagerAsyncOutput;
@@ -51,7 +52,7 @@ public:
 	void Terminate();
 	void SetInputMappings(const FInputNameMap& InNameMap)
 	{ 
-		FWriteScopeLock InputConfigLock(InputConfigurationLock);
+		FTransactionallySafeWriteScopeLock InputConfigLock(InputConfigurationLock);
 		InputNameMap = InNameMap; 
 	}
 
@@ -84,8 +85,8 @@ public:
 	/** Current control inputs that is being used on the PT */
 	FModularVehicleInputs VehicleInputs;
 	FInputNameMap InputNameMap;
-	FRWLock InputConfigurationLock;
+	FTransactionallySafeRWLock InputConfigurationLock;
 
 	int8 NetMode;
 
-};
\ No newline at end of file
+};
