Modular Vehicle refactor to work without cluster unions

diff --git a/Engine/Plugins/Experimental/ChaosModularVehicle/Source/ChaosModularVehicleEngine/Private/ChaosModularVehicle/ChaosSimModuleManagerAsyncCallback.cpp b/Engine/Plugins/Experimental/ChaosModularVehicle/Source/ChaosModularVehicleEngine/Private/ChaosModularVehicle/ChaosSimModuleManagerAsyncCallback.cpp
--- a/Engine/Plugins/Experimental/ChaosModularVehicle/Source/ChaosModularVehicleEngine/Private/ChaosModularVehicle/ChaosSimModuleManagerAsyncCallback.cpp
+++ b/Engine/Plugins/Experimental/ChaosModularVehicle/Source/ChaosModularVehicleEngine/Private/ChaosModularVehicle/ChaosSimModuleManagerAsyncCallback.cpp
@@ -200,15 +200,7 @@ void FModularVehicleAsyncInput::ApplyDeferredForces() const
 {
 	if (Vehicle && Proxy && Vehicle->VehicleSimulationPT)
 	{
-		if (Proxy->GetType() == EPhysicsProxyType::ClusterUnionProxy)
-		{
-			Vehicle->VehicleSimulationPT->ApplyDeferredForces(static_cast<Chaos::FClusterUnionPhysicsProxy*>(Proxy));
-		}
-		else if (Proxy->GetType() == EPhysicsProxyType::GeometryCollectionType)
-		{
-			Vehicle->VehicleSimulationPT->ApplyDeferredForces(static_cast<FGeometryCollectionPhysicsProxy*>(Proxy));
-		}
-
+		Vehicle->VehicleSimulationPT->ApplyDeferredForces(Proxy);
 	}
 
 }
@@ -220,7 +212,7 @@ void FModularVehicleAsyncInput::ProcessInputs()
 		return;
 	}
 
-	FModularVehicleSimulationCU* VehicleSim = GetVehicle()->VehicleSimulationPT.Get();
+	FModularVehicleSimulation* VehicleSim = GetVehicle()->VehicleSimulationPT.Get();
 
 	if (VehicleSim == nullptr || !GetVehicle()->bUsingNetworkPhysicsPrediction || GetVehicle()->GetWorld() == nullptr)
 	{
@@ -264,7 +256,7 @@ void FNetworkModularVehicleInputs::ApplyData(UActorComponent* NetworkComponent)
 	{
 		if (UModularVehicleBaseComponent* ModularBaseComponent = Cast<UModularVehicleBaseComponent>(NetworkComponent))
 		{
-			if (FModularVehicleSimulationCU* VehicleSimulation = ModularBaseComponent->VehicleSimulationPT.Get())
+			if (FModularVehicleSimulation* VehicleSimulation = ModularBaseComponent->VehicleSimulationPT.Get())
 			{
 				VehicleSimulation->VehicleInputs = VehicleInputs;
 			}
@@ -278,7 +270,7 @@ void FNetworkModularVehicleInputs::BuildData(const UActorComponent* NetworkCompo
 	{
 		if (const UModularVehicleBaseComponent* ModularBaseComponent = Cast<const UModularVehicleBaseComponent>(NetworkComponent))
 		{
-			if (const FModularVehicleSimulationCU* VehicleSimulation = ModularBaseComponent->VehicleSimulationPT.Get())
+			if (const FModularVehicleSimulation* VehicleSimulation = ModularBaseComponent->VehicleSimulationPT.Get())
 			{
 				VehicleInputs = VehicleSimulation->VehicleInputs;
 			}
@@ -404,7 +396,7 @@ void FNetworkModularVehicleStates::ApplyData(UActorComponent* NetworkComponent)
 {
 	if (UModularVehicleBaseComponent* ModularBaseComponent = Cast<UModularVehicleBaseComponent>(NetworkComponent))
 	{
-		if (FModularVehicleSimulationCU* VehicleSimulation = ModularBaseComponent->VehicleSimulationPT.Get())
+		if (FModularVehicleSimulation* VehicleSimulation = ModularBaseComponent->VehicleSimulationPT.Get())
 		{
 			VehicleSimulation->AccessSimComponentTree()->SetSimState(ModuleData);
 		}
@@ -415,7 +407,7 @@ void FNetworkModularVehicleStates::BuildData(const UActorComponent* NetworkCompo
 {
 	if (NetworkComponent)
 	{
-		if (const FModularVehicleSimulationCU* VehicleSimulation = Cast<const UModularVehicleBaseComponent>(NetworkComponent)->VehicleSimulationPT.Get())
+		if (const FModularVehicleSimulation* VehicleSimulation = Cast<const UModularVehicleBaseComponent>(NetworkComponent)->VehicleSimulationPT.Get())
 		{
 			VehicleSimulation->GetSimComponentTree()->SetNetState(ModuleData);
 		}
diff --git a/Engine/Plugins/Experimental/ChaosModularVehicle/Source/ChaosModularVehicleEngine/Private/ChaosModularVehicle/ModularVehicleBaseComponent.cpp b/Engine/Plugins/Experimental/ChaosModularVehicle/Source/ChaosModularVehicleEngine/Private/ChaosModularVehicle/ModularVehicleBaseComponent.cpp
--- a/Engine/Plugins/Experimental/ChaosModularVehicle/Source/ChaosModularVehicleEngine/Private/ChaosModularVehicle/ModularVehicleBaseComponent.cpp
+++ b/Engine/Plugins/Experimental/ChaosModularVehicle/Source/ChaosModularVehicleEngine/Private/ChaosModularVehicle/ModularVehicleBaseComponent.cpp
@@ -18,6 +18,9 @@
 #include "GameFramework/Pawn.h"
 #include "Physics/Experimental/PhysScene_Chaos.h"
 #include "SimModule/SimModuleTree.h"
+#include "PhysicsProxy/SingleParticlePhysicsProxy.h"
+#include "PhysicsProxy/ClusterUnionPhysicsProxy.h"
+#include "GeometryCollection/GeometryCollectionComponent.h"
 
 #include "ChaosModularVehicle/ChaosSimModuleManagerAsyncCallback.h"
 
@@ -220,13 +223,50 @@ void UModularVehicleBaseComponent::OnDestroyPhysicsState()
 
 }
 
-
-int GenerateNewGuid()
+int UModularVehicleBaseComponent::GenerateNewGuid()
 {
 	static int Val = 0;
 	return Val++;
 }
 
+void UModularVehicleBaseComponent::CacheRootPhysicsObject(IPhysicsProxyBase* Proxy)
+{
+	Chaos::EnsureIsInGameThreadContext();
+	using namespace Chaos;
+	RootPhysicsObject = nullptr;
+	
+	if (Proxy == nullptr)
+	{
+		return;
+	}
+
+	switch (Proxy->GetType())
+	{
+		case EPhysicsProxyType::ClusterUnionProxy:
+		{
+			if (FClusterUnionPhysicsProxy* CUProxy = static_cast<FClusterUnionPhysicsProxy*>(Proxy))
+			{
+				RootPhysicsObject = CUProxy->GetPhysicsObjectHandle();
+			}
+		}
+		break;
+
+		case EPhysicsProxyType::SingleParticleProxy:
+		{
+			if (FSingleParticlePhysicsProxy* ParticleProxy = static_cast<FSingleParticlePhysicsProxy*>(Proxy))
+			{
+				RootPhysicsObject = ParticleProxy->GetPhysicsObject();
+			}
+		}
+		break;
+
+		default:
+		{
+			UE_LOG(LogModularBase, Error, TEXT("Unsupported Particle type"));
+		}
+		break;
+	}
+}
 
 void UModularVehicleBaseComponent::CreateAssociatedSimComponents(USceneComponent* AttachedComponent, int ParentIndex, int TransformIndex, Chaos::FSimTreeUpdates& TreeUpdatesOut)
 {
@@ -236,6 +276,11 @@ void UModularVehicleBaseComponent::CreateAssociatedSimComponents(USceneComponent
 		return;
 	}
 
+	if (RootPhysicsObject == nullptr)
+	{
+		CacheRootPhysicsObject(GetPhysicsProxy());
+	}
+
 	UE_LOG(LogModularBase, Log, TEXT("CreateAssociatedSimComponents Attaching %s, TransformIndex %d"), *AttachedComponent->GetName(), TransformIndex);
 
 	TArray<Chaos::FClusterUnionChildData> ChildParticles = ClusterUnionComponent->GetPhysicsProxy()->GetSyncedData_External().ChildParticles;
@@ -286,7 +331,126 @@ void UModularVehicleBaseComponent::CreateAssociatedSimComponents(USceneComponent
 
 		const bool bIsAnimationEnabled = ComponentInterface->GetAnimationEnabled();
 		const FName ComponentBoneName = ComponentInterface->GetBoneName();
-		if (bIsAnimationEnabled && (ComponentBoneName != NAME_None))
+		// Now always using ModuleAnimationSetups even when there is no bone specified. ModuleAnimationSetups can be used to drive
+		// transforms in components, not just skeletal mesh bones.
+		if (bIsAnimationEnabled)
+		{
+			// if bone already exists then use that (seperate wheel and suspension modules can share same bone)
+
+			int FoundIndex = -1;
+			if (ComponentBoneName != NAME_None)
+			{
+				for (int I = 0; I < ModuleAnimationSetups.Num(); I++)
+				{
+					if (ModuleAnimationSetups[I].BoneName == ComponentBoneName)
+					{
+						FoundIndex = I;
+						break;
+					}
+				}
+			}
+
+			const FVector& ComponentAnimationOffset = ComponentInterface->GetAnimationOffset();
+			if (FoundIndex != -1)
+			{
+				NewModule->SetAnimationData(ComponentBoneName, ComponentAnimationOffset, FoundIndex);
+			}
+			else
+			{
+				NewModule->SetAnimationData(ComponentBoneName, ComponentAnimationOffset, ModuleAnimationSetups.Num());
+				FModuleAnimationSetup AnimSetup(NewModule->GetBoneName(), TransformIndex);
+				ModuleAnimationSetups.Add(AnimSetup);
+			}
+		}
+
+		// store the tree index in the original sim component
+		ComponentInterface->SetTreeIndex(TreeIndex);
+		ParentIndex = TreeIndex;
+
+		if (GetPhysicsProxy())
+		{
+			NewModule->OnConstruction_External(RootPhysicsObject);
+		}
+	}
+
+	// get this components children, not all children
+	TArray<USceneComponent*> Children;
+	AttachedComponent->GetChildrenComponents(false, Children);
+
+	// recurse down tree, converting all SimComponents to proper simulation modules
+	for (USceneComponent* Child : Children)
+	{
+		if (IVehicleSimBaseComponentInterface* ChildSimComponent = Cast<IVehicleSimBaseComponentInterface>(Child))
+		{
+			CreateAssociatedSimComponents(Child, ParentIndex, TransformIndex, TreeUpdatesOut);
+		}
+	}
+
+	UpdatePhysicalProperties();
+}
+
+void UModularVehicleBaseComponent::CreateIndependentSimComponents(USceneComponent* RootComponent, USceneComponent* AttachedComponent, int ParentIndex, int TransformIndex, Chaos::FSimTreeUpdates& TreeUpdatesOut)
+{
+	using namespace Chaos;
+	if (AttachedComponent == nullptr || RootComponent == nullptr)
+	{
+		return;
+	}
+
+	if (RootPhysicsObject == nullptr)
+	{
+		CacheRootPhysicsObject(GetPhysicsProxy()); //#TODO where to put / refresh?
+	}
+
+	UE_LOG(LogModularBase, Log, TEXT("CreateAssociatedSimComponents Attaching %s, TransformIndex %d"), *AttachedComponent->GetName(), TransformIndex);
+
+	if (IVehicleSimBaseComponentInterface* ComponentInterface = Cast<IVehicleSimBaseComponentInterface>(AttachedComponent))
+	{
+		FTransform RootComponentTransform = RootComponent->GetComponentTransform();
+		FTransform ComponentTransform = AttachedComponent->GetComponentTransform().GetRelativeTransform(RootComponentTransform);
+
+		int TreeIndex = INDEX_NONE;
+
+		Chaos::ISimulationModuleBase* NewModule = ComponentInterface->CreateNewCoreModule();
+
+		TUniquePtr<Chaos::FSimModuleTree>& SimModuleTree = VehicleSimulationPT->AccessSimComponentTree();
+
+#if !(UE_BUILD_SHIPPING || UE_BUILD_TEST)
+		FString DebugString;
+		NewModule->GetDebugString(DebugString);
+		UE_LOG(LogModularBase, Log, TEXT("CreateAssociatedSimComponents Found Sim Component %s (%s)"), *NewModule->GetDebugName(), *DebugString);
+#endif
+
+		FVector LocalOffset(0, 0, 0);
+		FVector LocalOffsetCOM(0, 0, 0);
+		FTransform PhysicsTransform = FTransform::Identity;
+
+		FTransform InitialTransform = ComponentTransform;
+		InitialTransform.SetLocation(InitialTransform.GetLocation());
+		NewModule->SetSimModuleTree(SimModuleTree.Get());
+		TreeIndex = TreeUpdatesOut.AddNodeBelow(ParentIndex, NewModule);
+		NewModule->SetTransformIndex(TransformIndex);
+		NewModule->SetParticleIndex(Chaos::FUniqueIdx());
+
+		NewModule->SetIntactTransform(FTransform::Identity);
+		NewModule->SetClusteredTransform(InitialTransform);
+		NewModule->SetClustered(true);
+
+		FVehicleComponentData& ComponentData = ComponentToPhysicsObjects.FindOrAdd(AttachedComponent);
+		ComponentData.Guid = GenerateNewGuid();
+		NewModule->SetGuid(ComponentData.Guid);
+
+		NewModule->SetInitialParticleTransform(InitialTransform);
+		NewModule->SetComponentTransform(ComponentTransform);
+
+		FTransform ClusterredTransform(FQuat::Identity, InitialTransform.GetLocation());
+		NewModule->SetClusteredTransform(ClusterredTransform);
+
+		const bool bIsAnimationEnabled = ComponentInterface->GetAnimationEnabled();
+		const FName ComponentBoneName = ComponentInterface->GetBoneName();
+		// Changed behaviour - now always using ModuleAnimationSetups even when there is no bone setup. ModuleAnimationSetups is used to drive
+		// transforms in components, not just skeletal mesh bones.
+		if (bIsAnimationEnabled)
 		{
 			// if bone already exists then use that (seperate wheel and suspension modules can share same bone)
 
@@ -308,7 +472,7 @@ void UModularVehicleBaseComponent::CreateAssociatedSimComponents(USceneComponent
 			else
 			{
 				NewModule->SetAnimationData(ComponentBoneName, ComponentAnimationOffset, ModuleAnimationSetups.Num());
-				FModuleAnimationSetup AnimSetup(NewModule->GetBoneName());
+				FModuleAnimationSetup AnimSetup(NewModule->GetBoneName(), TransformIndex);
 				ModuleAnimationSetups.Add(AnimSetup);
 			}
 		}
@@ -317,9 +481,9 @@ void UModularVehicleBaseComponent::CreateAssociatedSimComponents(USceneComponent
 		ComponentInterface->SetTreeIndex(TreeIndex);
 		ParentIndex = TreeIndex;
 
-		if (Chaos::FClusterUnionPhysicsProxy* Proxy = ClusterUnionComponent->GetPhysicsProxy())
+		if (GetPhysicsProxy())
 		{
-			NewModule->OnConstruction_External(Proxy);
+			NewModule->OnConstruction_External(RootPhysicsObject);
 		}
 	}
 
@@ -332,30 +496,91 @@ void UModularVehicleBaseComponent::CreateAssociatedSimComponents(USceneComponent
 	{
 		if (IVehicleSimBaseComponentInterface* ChildSimComponent = Cast<IVehicleSimBaseComponentInterface>(Child))
 		{
-			CreateAssociatedSimComponents(Child, ParentIndex, TransformIndex, TreeUpdatesOut);
+			CreateIndependentSimComponents(RootComponent, Child, ParentIndex, TransformIndex, TreeUpdatesOut);
 		}
 	}
 
 	UpdatePhysicalProperties();
 }
 
+
+int UModularVehicleBaseComponent::AddSimModule(Chaos::ISimulationModuleBase* NewModule, const FTransform& LocalTransform, int ParentIndex, int TransformIndex)
+{
+	TUniquePtr<Chaos::FSimModuleTree>& SimModuleTree = VehicleSimulationPT->AccessSimComponentTree();
+
+#if !(UE_BUILD_SHIPPING || UE_BUILD_TEST)
+	FString DebugString;
+	NewModule->GetDebugString(DebugString);
+	UE_LOG(LogModularBase, Log, TEXT("CreateAssociatedSimComponents Found Sim Component %s (%s)"), *NewModule->GetDebugName(), *DebugString);
+#endif
+
+	FVector LocalOffset(0, 0, 0);
+	FVector LocalOffsetCOM(0, 0, 0);
+	FTransform PhysicsTransform = FTransform::Identity;
+
+	FTransform InitialTransform = LocalTransform;
+	InitialTransform.SetLocation(InitialTransform.GetLocation());
+	NewModule->SetSimModuleTree(SimModuleTree.Get());
+	int TreeIndex = StoredTreeUpdates.AddNodeBelow(ParentIndex, NewModule);
+
+	NewModule->SetGuid(GenerateNewGuid());
+
+	NewModule->SetTransformIndex(TransformIndex);
+	NewModule->SetParticleIndex(Chaos::FUniqueIdx());
+
+	NewModule->SetIntactTransform(FTransform::Identity);
+	NewModule->SetClusteredTransform(InitialTransform);
+	NewModule->SetClustered(true);
+
+	NewModule->SetInitialParticleTransform(InitialTransform);
+	NewModule->SetComponentTransform(LocalTransform);
+
+	FTransform ClusterredTransform(FQuat::Identity, InitialTransform.GetLocation());
+	NewModule->SetClusteredTransform(ClusterredTransform);
+
+	return TreeIndex;
+}
+
+void UModularVehicleBaseComponent::RemoveSimModule(int ModuleGuid)
+{
+	StoredTreeUpdates.RemoveNode(ModuleGuid);
+}
+
+void UModularVehicleBaseComponent::FinalizeModuleUpdates()
+{
+	ActionTreeUpdates(&StoredTreeUpdates);
+	StoredTreeUpdates.ClearUpdates();
+
+	if (RootPhysicsObject == nullptr)
+	{
+		CacheRootPhysicsObject(GetPhysicsProxy()); //#TODO where to put / refresh?
+	}
+
+	UpdatePhysicalProperties();
+}
+
+
 void UModularVehicleBaseComponent::UpdatePhysicalProperties()
 {
-	if (Chaos::FClusterUnionPhysicsProxy* Proxy = ClusterUnionComponent->GetPhysicsProxy())
+	if (ClusterUnionComponent)
 	{
-		Chaos::FPBDRigidsSolver* Solver = Proxy->GetSolver<Chaos::FPBDRigidsSolver>();
-		Solver->EnqueueCommandImmediate([Proxy, this]() mutable
-			{
-				Proxy->GetParticle_Internal()->SetLinearEtherDrag(LinearDamping);
-				Proxy->GetParticle_Internal()->SetAngularEtherDrag(AngularDamping);
-			});
+		if (Chaos::FClusterUnionPhysicsProxy* Proxy = ClusterUnionComponent->GetPhysicsProxy())
+		{
+			Chaos::FPBDRigidsSolver* Solver = Proxy->GetSolver<Chaos::FPBDRigidsSolver>();
+			Solver->EnqueueCommandImmediate([Proxy, this]() mutable
+				{
+					Proxy->GetParticle_Internal()->SetLinearEtherDrag(LinearDamping);
+					Proxy->GetParticle_Internal()->SetAngularEtherDrag(AngularDamping);
+				});
+		}
+	}
 
-		// updating external particle currently doesn't update physics particle
-		//	if (auto* Particle = static_cast<Chaos::FClusterUnionPhysicsProxy*>(Proxy)->GetParticle_External())
-		//	{
-		//		Particle->SetLinearEtherDrag(LinearDamping);
-		//		Particle->SetAngularEtherDrag(AngularDamping);
-		//	}
+	if (RootPhysicsObject)
+	{
+		TArray<Chaos::FPhysicsObjectHandle> PhysicsObjects{ RootPhysicsObject };
+		FLockedWritePhysicsObjectExternalInterface Interface = FPhysicsObjectExternalInterface::LockWrite(PhysicsObjects);
+		Interface->SetLinearEtherDrag(PhysicsObjects, LinearDamping);
+		Interface->SetAngularEtherDrag(PhysicsObjects, AngularDamping);
 	}
 }
 
@@ -376,6 +601,8 @@ void UModularVehicleBaseComponent::BeginPlay()
 		}
 	}
 
+	//AddSimulationComponentsFromOwnedActor();
+	
 	// control input setup - unfortunately can't do this in OnCreatePhysics since RootComponent->GetChildrenComponents will not work 
 	// at that time and AssimilateComponentInputs will not find any controls in the component hierarchy
 	TArray<FModuleInputSetup> CombinedInputConfiguration;
@@ -416,7 +643,7 @@ void UModularVehicleBaseComponent::TickComponent(float DeltaTime, ELevelTick Tic
 }
 
 
-int32 UModularVehicleBaseComponent::FindComponentAddOrder(UPrimitiveComponent* InComponent)
+int32 UModularVehicleBaseComponent::FindComponentAddOrder(USceneComponent* InComponent)
 {
 	int32 ComponentAddOrder = INDEX_NONE;
 	for (const FConstructionData& Data : ConstructionDatas)
@@ -433,7 +660,7 @@ int32 UModularVehicleBaseComponent::FindComponentAddOrder(UPrimitiveComponent* I
 bool UModularVehicleBaseComponent::FindAndRemoveNextPendingUpdate(int32 NextIndex, Chaos::FSimTreeUpdates* OutData)
 {
 	bool bFound = false;
-	for (const TPair<TObjectKey<UPrimitiveComponent>, Chaos::FSimTreeUpdates>& Kvp : PendingTreeUpdates)
+	for (const TPair<TObjectKey<USceneComponent>, Chaos::FSimTreeUpdates>& Kvp : PendingTreeUpdates)
 	{
 		if (FindComponentAddOrder(Kvp.Key.ResolveObjectPtr()) == NextIndex)
 		{
@@ -609,7 +836,27 @@ void UModularVehicleBaseComponent::SetCurrentAsyncDataInternal(FModularVehicleAs
 	NewOutput.Idx = InputIdx;
 }
 
-inline IPhysicsProxyBase* UModularVehicleBaseComponent::GetPhysicsProxy() const { return ClusterUnionComponent ? ClusterUnionComponent->GetPhysicsProxy() : nullptr; }
+IPhysicsProxyBase* UModularVehicleBaseComponent::GetPhysicsProxy() const 
+{ 
+	if (CachedPhysicsProxy)
+	{
+		return CachedPhysicsProxy;
+	}
+
+	// ClusterUnionPhysicsProxy case
+	if (ClusterUnionComponent)
+	{
+		return CachedPhysicsProxy = ClusterUnionComponent->GetPhysicsProxy();
+	}
+	
+	// SingleParticlePhysicsProxy case
+	if (const FBodyInstance* BodyInstance = UpdatedPrimitive ? UpdatedPrimitive->GetBodyInstance() : nullptr)
+	{
+		return CachedPhysicsProxy = BodyInstance->ActorHandle;
+	}
+
+	return nullptr;
+}
 
 
 void UModularVehicleBaseComponent::SetClusterComponent(UClusterUnionComponent* InPhysicalComponent)
@@ -680,18 +927,21 @@ void UModularVehicleBaseComponent::ParallelUpdate(float DeltaTime)
 
 				if (Chaos::FSimOutputData* ModuleOutput = PVehicleOutput->SimTreeOutputData[I])
 				{
-					if ((ModuleOutput->AnimationSetupIndex >= 0) && (ModuleOutput->AnimationSetupIndex < ModuleAnimationSetups.Num()))
+					const Chaos::FSimModuleAnimationData& AnimData = ModuleOutput->AnimationData;
+					int Index = AnimData.AnimationSetupIndex;
+					if ((Index >= 0) && (Index < ModuleAnimationSetups.Num()))
 					{
-						ModuleAnimationSetups[ModuleOutput->AnimationSetupIndex].AnimFlags |= ModuleOutput->AnimFlags;
+						ModuleAnimationSetups[Index].AnimFlags |= AnimData.AnimFlags;
+						ModuleAnimationSetups[Index].CombinedRotation = AnimData.CombinedRotation;
 
-						if (ModuleOutput->AnimFlags & Chaos::EAnimationFlags::AnimateRotation)
+						if (AnimData.AnimFlags & Chaos::EAnimationFlags::AnimateRotation)
 						{
-							ModuleAnimationSetups[ModuleOutput->AnimationSetupIndex].RotOffset = ModuleOutput->AnimationRotOffset;
+							ModuleAnimationSetups[Index].RotOffset = AnimData.AnimationRotOffset;
 						}
 
-						if (ModuleOutput->AnimFlags & Chaos::EAnimationFlags::AnimatePosition)
+						if (AnimData.AnimFlags & Chaos::EAnimationFlags::AnimatePosition)
 						{
-							ModuleAnimationSetups[ModuleOutput->AnimationSetupIndex].LocOffset = ModuleOutput->AnimationLocOffset;
+							ModuleAnimationSetups[Index].LocOffset = AnimData.AnimationLocOffset;
 						}
 					}
 				}
@@ -710,15 +960,16 @@ void UModularVehicleBaseComponent::ParallelUpdate(float DeltaTime)
 	}
 }
 
+
 /************************************************************************/
 /* PASS ANY INPUTS TO THE PHYSICS THREAD SIMULATION IN HERE             */
 /************************************************************************/
 void UModularVehicleBaseComponent::Update(float DeltaTime)
 {
-	if (CurAsyncInput && ClusterUnionComponent && ClusterUnionComponent->GetPhysicsProxy())
+	if (CurAsyncInput && GetPhysicsProxy())
 	{
 
-		CurAsyncInput->Proxy = ClusterUnionComponent->GetPhysicsProxy();
+		CurAsyncInput->Proxy = GetPhysicsProxy();
 
 		FModularVehicleAsyncInput* AsyncInput = static_cast<FModularVehicleAsyncInput*>(CurAsyncInput);
 
@@ -754,64 +1005,61 @@ const FTransform& UModularVehicleBaseComponent::GetComponentTransform() const
 
 void UModularVehicleBaseComponent::ActionTreeUpdates(Chaos::FSimTreeUpdates* NextTreeUpdates)
 {
-	if (ClusterUnionComponent)
+	if (IPhysicsProxyBase* Proxy = GetPhysicsProxy())
 	{
-		if (Chaos::FClusterUnionPhysicsProxy* Proxy = ClusterUnionComponent->GetPhysicsProxy())
+		if (Chaos::FPBDRigidsSolver* Solver = Proxy->GetSolver<Chaos::FPBDRigidsSolver>())
 		{
-			if (Chaos::FPBDRigidsSolver* Solver = Proxy->GetSolver<Chaos::FPBDRigidsSolver>())
-			{
-	Solver->EnqueueCommandImmediate([Proxy, this, NextTreeUpdates = *NextTreeUpdates]() mutable
+Solver->EnqueueCommandImmediate([Proxy, this, NextTreeUpdates = *NextTreeUpdates]() mutable
+	{
+		if (IsValid(this) && bPhysicsStateCreated && VehicleSimulationPT)
 		{
-			if (IsValid(this) && bPhysicsStateCreated && VehicleSimulationPT)
+			TUniquePtr<Chaos::FSimModuleTree>& SimModuleTree = VehicleSimulationPT->AccessSimComponentTree();
+			if(SimModuleTree.IsValid())
 			{
-				TUniquePtr<Chaos::FSimModuleTree>& SimModuleTree = VehicleSimulationPT->AccessSimComponentTree();
-				if(SimModuleTree.IsValid())
-				{
-					SimModuleTree->AppendTreeUpdates(NextTreeUpdates);
-					FModularVehicleBuilder::FixupTreeLinks(SimModuleTree);
+				SimModuleTree->AppendTreeUpdates(NextTreeUpdates);
+				FModularVehicleBuilder::FixupTreeLinks(SimModuleTree);
 
-	#if !(UE_BUILD_SHIPPING || UE_BUILD_TEST)
-					if (bModularVehicle_DumpModuleTreeStructure_Enabled)
+#if !(UE_BUILD_SHIPPING || UE_BUILD_TEST)
+				if (bModularVehicle_DumpModuleTreeStructure_Enabled)
+				{
+					UE_LOG(LogModularBase, Warning, TEXT("SimTreeModules:") );
+					for (int I = 0; I < SimModuleTree->GetNumNodes(); I++)
 					{
-						UE_LOG(LogTemp, Warning, TEXT("SimTreeModules:") );
-						for (int I = 0; I < SimModuleTree->GetNumNodes(); I++)
+						if (Chaos::ISimulationModuleBase* Module = SimModuleTree->GetNode(I).SimModule)
 						{
-							if (Chaos::ISimulationModuleBase* Module = SimModuleTree->GetNode(I).SimModule)
-							{
-								FString String;
-								Module->GetDebugString(String);
-								UE_LOG(LogTemp, Warning, TEXT("..%s"), *String);
-							}
+							FString String;
+							Module->GetDebugString(String);
+							UE_LOG(LogModularBase, Warning, TEXT("..%s"), *String);
 						}
 					}
-	#endif // !(UE_BUILD_SHIPPING || UE_BUILD_TEST)
+				}
+#endif // !(UE_BUILD_SHIPPING || UE_BUILD_TEST)
 
-					// Network replication data needs to be updated, this is currently studily slow
-					if (NetworkPhysicsComponent)
+				// Network replication data needs to be updated, this is currently studily slow
+				if (NetworkPhysicsComponent)
+				{
+					TSharedPtr<Chaos::FBaseRewindHistory>& History = NetworkPhysicsComponent->GetStateHistory_Internal();
+					Chaos::TDataRewindHistory<FNetworkModularVehicleStates>* StateHistory = static_cast<Chaos::TDataRewindHistory<FNetworkModularVehicleStates>*>(History.Get());
+					if (StateHistory)
 					{
-						TSharedPtr<Chaos::FBaseRewindHistory>& History = NetworkPhysicsComponent->GetStateHistory_Internal();
-						Chaos::TDataRewindHistory<FNetworkModularVehicleStates>* StateHistory = static_cast<Chaos::TDataRewindHistory<FNetworkModularVehicleStates>*>(History.Get());
-						if (StateHistory)
+						// #TODO: we are rebuilding from scratch every time there is a single change, there must be a better way!
+						// not sure of it is safe to update the data at this time?
+						for (int I = 0; I < StateHistory->GetDataHistory().Num(); I++)
 						{
-							// #TODO: we are rebuilding from scratch every time there is a single change, there must be a better way!
-							// not sure of it is safe to update the data at this time?
-							for (int I = 0; I < StateHistory->GetDataHistory().Num(); I++)
+							FNetworkModularVehicleStates& State = StateHistory->GetDataHistory()[I];
+							State.ModuleData.Empty();
+							TUniquePtr<Chaos::FSimModuleTree>& InnerSimModuleTree = VehicleSimulationPT->AccessSimComponentTree();
+							if(InnerSimModuleTree.IsValid())
 							{
-								FNetworkModularVehicleStates& State = StateHistory->GetDataHistory()[I];
-								State.ModuleData.Empty();
-								TUniquePtr<Chaos::FSimModuleTree>& InnerSimModuleTree = VehicleSimulationPT->AccessSimComponentTree();
-								if(InnerSimModuleTree.IsValid())
-								{
-									InnerSimModuleTree->GenerateReplicationStructure(State.ModuleData);
-								}
+								InnerSimModuleTree->GenerateReplicationStructure(State.ModuleData);
 							}
 						}
 					}
 				}
 			}
-		});
-}
 		}
+	});
+}
 	}
 }
 
@@ -835,6 +1083,11 @@ int32 UModularVehicleBaseComponent::FindParentsLastSimComponent(const USceneComp
 }
 
 void UModularVehicleBaseComponent::AddComponentToSimulation(UPrimitiveComponent* InComponent, const TArray<FClusterUnionBoneData>& BonesData, const TArray<FClusterUnionBoneData>& RemovedBoneIDs, bool bIsNew)
+{
+	AddComponentToSimulationImpl(InComponent, BonesData, RemovedBoneIDs, bIsNew);
+}
+
+void UModularVehicleBaseComponent::AddComponentToSimulationImpl(UPrimitiveComponent* InComponent, const TArray<FClusterUnionBoneData>& BonesData, const TArray<FClusterUnionBoneData>& RemovedBoneIDs, bool bIsNew)
 {
 	check(ClusterUnionComponent);
 
@@ -881,6 +1134,11 @@ void UModularVehicleBaseComponent::AddComponentToSimulation(UPrimitiveComponent*
 }
 
 void UModularVehicleBaseComponent::RemoveComponentFromSimulation(UPrimitiveComponent* InComponent, const TArray<FClusterUnionBoneData>& RemovedBonesData)
+{
+	RemoveComponentFromSimulationImpl(InComponent, RemovedBonesData);
+}
+
+void UModularVehicleBaseComponent::RemoveComponentFromSimulationImpl(UPrimitiveComponent* InComponent, const TArray<FClusterUnionBoneData>& RemovedBonesData)
 {
 	if (InComponent && VehicleSimulationPT)
 	{
@@ -1003,7 +1261,7 @@ void UModularVehicleBaseComponent::CreateVehicleSim()
 		if (World->IsGameWorld())
 		{
 			// create the simulation class
-			VehicleSimulationPT = MakeUnique<FModularVehicleSimulationCU>(bUsingNetworkPhysicsPrediction, (int8)World->GetNetMode());
+			VehicleSimulationPT = MakeUnique<FModularVehicleSimulation>(bUsingNetworkPhysicsPrediction, (int8)World->GetNetMode());
 
 			// create physics output container
 			PVehicleOutput = MakeUnique<FPhysicsVehicleOutput>();
@@ -1055,34 +1313,36 @@ void UModularVehicleBaseComponent::DestroyVehicleSim()
 
 void UModularVehicleBaseComponent::AddOverlappingComponentsToCluster()
 {
-	if (UWorld* World = GetWorld())
-	{
-		FVector HalfSize = AutoAddOverlappingBoxSize * 0.5f;
-		FBox LocalBox = FBox(-HalfSize, HalfSize);
-		TArray<FOverlapResult> OverlapResults;
-		FCollisionShape CollisionBox;
-		CollisionBox.SetBox((FVector3f)LocalBox.GetExtent());
-
-		const FCollisionQueryParams QueryParams;
-		const FCollisionResponseParams ResponseParams;
-		const ECollisionChannel CollisionChannel = ECollisionChannel::ECC_WorldDynamic;
-		const bool bOverlapHit = World->OverlapMultiByChannel(OverlapResults, GetActorLocation(), FQuat::Identity, CollisionChannel, CollisionBox, QueryParams, ResponseParams);
-
-		TArray<int32> BoneIds;
-		BoneIds.Add(0);
-		// Test each overlapped object for a hit result
-		for (FOverlapResult OverlapResult : OverlapResults)
-		{
-			if (UGeometryCollectionComponent* OverlapGCComponent = Cast<UGeometryCollectionComponent>(OverlapResult.Component.Get()))
-			{
-				ClusterUnionComponent->AddComponentToCluster(OverlapGCComponent, BoneIds);
-			}
-		}
-	}
+	//if (UWorld* World = GetWorld())
+	//{
+	//	FVector HalfSize = AutoAddOverlappingBoxSize * 0.5f;
+	//	FBox LocalBox = FBox(-HalfSize, HalfSize);
+	//	TArray<FOverlapResult> OverlapResults;
+	//	FCollisionShape CollisionBox;
+	//	CollisionBox.SetBox((FVector3f)LocalBox.GetExtent());
+
+	//	const FCollisionQueryParams QueryParams;
+	//	const FCollisionResponseParams ResponseParams;
+	//	const ECollisionChannel CollisionChannel = ECollisionChannel::ECC_WorldDynamic;
+	//	const bool bOverlapHit = World->OverlapMultiByChannel(OverlapResults, GetActorLocation(), FQuat::Identity, CollisionChannel, CollisionBox, QueryParams, ResponseParams);
+
+	//	TArray<int32> BoneIds;
+	//	BoneIds.Add(0);
+	//	// Test each overlapped object for a hit result
+	//	for (FOverlapResult OverlapResult : OverlapResults)
+	//	{
+	//		if (UGeometryCollectionComponent* OverlapGCComponent = Cast<UGeometryCollectionComponent>(OverlapResult.Component.Get()))
+	//		{
+	//			ClusterUnionComponent->AddComponentToCluster(OverlapGCComponent, BoneIds);
+	//		}
+	//	}
+	//}
 }
 
-void UModularVehicleBaseComponent::AddGeometryCollectionsFromOwnedActor()
+bool UModularVehicleBaseComponent::AddGeometryCollectionsFromOwnedActor()
 {
+	bool bSuccess = false;
+
 	TArray<int32> BoneIds;
 	BoneIds.Add(0);
 
@@ -1098,10 +1358,44 @@ void UModularVehicleBaseComponent::AddGeometryCollectionsFromOwnedActor()
 				if (UGeometryCollectionComponent* GCComponent = Cast<UGeometryCollectionComponent>(Component))
 				{
 					ClusterUnionComponent->AddComponentToCluster(GCComponent, BoneIds);
+					bSuccess = true;
 				}
 			}
 		}
 	}
+
+	return bSuccess;
+}
+
+void UModularVehicleBaseComponent::AddSimulationComponentsFromOwnedActor()
+{
+	// get raw simulation components that are not connected with any component - required for scene graph
+	if (APawn* Pawn = Cast<APawn>(GetOwner()))
+	{
+		USceneComponent* RootComponent = Pawn->GetRootComponent();
+		TArray<USceneComponent*> ChildComponents;
+		RootComponent->GetChildrenComponents(false, ChildComponents);
+
+		int32 ComponentAddOrder = INDEX_NONE;
+		int ParentIndex=-1;
+		int TransformIndex=0;
+		Chaos::FSimTreeUpdates LatestTreeUpdates;
+
+		for (USceneComponent* Component : ChildComponents)
+		{
+			if (UVehicleSimBaseComponent* SimulationComponent = Cast<UVehicleSimBaseComponent>(Component))
+			{
+				CreateIndependentSimComponents(RootComponent, SimulationComponent, ParentIndex, TransformIndex, LatestTreeUpdates);
+				NextTransformIndex++;
+			}
+		}
+
+		Chaos::FSimTreeUpdates* NextTreeUpdates = nullptr;
+		NextTreeUpdates = &LatestTreeUpdates;
+		LastComponentAddIndex = ComponentAddOrder;
+		ActionTreeUpdates(NextTreeUpdates);
+
+	}
 }
 
 void UModularVehicleBaseComponent::SetInputProducerClass(TSubclassOf<UVehicleInputProducerBase> InInputProducerClass)
@@ -1236,7 +1530,7 @@ void UModularVehicleBaseComponent::LogInputSetup()
 
 	for (auto& NamePair : NameMap)
 	{ 
-		UE_LOG(LogTemp, Warning, TEXT("Input: %s %d"), *NamePair.Key.ToString(), NamePair.Value);
+		UE_LOG(LogModularBase, Log, TEXT("Input: %s %d"), *NamePair.Key.ToString(), NamePair.Value);
 	}
 #endif //!(UE_BUILD_SHIPPING || UE_BUILD_TEST)
 }
\ No newline at end of file
diff --git a/Engine/Plugins/Experimental/ChaosModularVehicle/Source/ChaosModularVehicleEngine/Private/ChaosModularVehicle/ModularVehicleSimulationCU.cpp b/Engine/Plugins/Experimental/ChaosModularVehicle/Source/ChaosModularVehicleEngine/Private/ChaosModularVehicle/ModularVehicleSimulationCU.cpp
--- a/Engine/Plugins/Experimental/ChaosModularVehicle/Source/ChaosModularVehicleEngine/Private/ChaosModularVehicle/ModularVehicleSimulationCU.cpp
+++ b/Engine/Plugins/Experimental/ChaosModularVehicle/Source/ChaosModularVehicleEngine/Private/ChaosModularVehicle/ModularVehicleSimulationCU.cpp
@@ -4,17 +4,18 @@
 #include "ChaosModularVehicle/ModularVehicleDefaultAsyncInput.h"
 #include "SimModule/SimModulesInclude.h"
 #include "SimModule/ModuleInput.h"
-#include "PhysicsProxy/GeometryCollectionPhysicsProxy.h"
 #include "PhysicsProxy/ClusterUnionPhysicsProxy.h"
-#include "GeometryCollection/GeometryCollectionObject.h"
+#include "PhysicsProxy/SingleParticlePhysicsProxy.h"
 #include "PhysicalMaterials/PhysicalMaterial.h"
 #include "PhysicsEngine/PhysicsSettings.h"
 #include "Chaos/ClusterUnionManager.h"
 #include "Chaos/DebugDrawQueue.h"
 #include "Engine/World.h"
+#include "Engine/HitResult.h"
 #include "ChaosModularVehicle/ModularVehicleDebug.h"
 
 FModularVehicleDebugParams GModularVehicleDebugParams;
+DEFINE_LOG_CATEGORY(LogModularVehicleSim);
 
 #if CHAOS_DEBUG_DRAW
 FAutoConsoleVariableRef CVarChaosModularVehiclesRaycastsEnabled(TEXT("p.ModularVehicle.SuspensionRaycastsEnabled"), GModularVehicleDebugParams.SuspensionRaycastsEnabled, TEXT("Enable/Disable Suspension Raycasts."));
@@ -27,18 +28,23 @@ FAutoConsoleVariableRef CVarChaosModularVehiclesDisableAnim(TEXT("p.ModularVehic
 #endif
 
 
-void FModularVehicleSimulationCU::Initialize(TUniquePtr<Chaos::FSimModuleTree>& InSimModuleTree)
+void FModularVehicleSimulation::Initialize(TUniquePtr<Chaos::FSimModuleTree>& InSimModuleTree)
 {
 	SimModuleTree = MoveTemp(InSimModuleTree);
 }
 
-void FModularVehicleSimulationCU::Terminate()
+void FModularVehicleSimulation::Terminate()
 {
 	SimModuleTree.Reset(nullptr);
 }
 
-void FModularVehicleSimulationCU::Simulate(UWorld* InWorld, float DeltaSeconds, const FModularVehicleAsyncInput& InputData, FModularVehicleAsyncOutput& OutputData, IPhysicsProxyBase* Proxy)
+void FModularVehicleSimulation::Simulate(UWorld* InWorld, float DeltaSeconds, const FModularVehicleAsyncInput& InputData, FModularVehicleAsyncOutput& OutputData, IPhysicsProxyBase* Proxy)
 {
+	if (RootParticle == nullptr)
+	{
+		CacheRootParticle(Proxy);
+	}
+	
 	Chaos::FPBDRigidsSolver* RigidsSolver = static_cast<Chaos::FPBDRigidsSolver*>(Proxy->GetSolver<Chaos::FPhysicsSolver>());
 	int CurrentFrame = -1;
 	if (RigidsSolver != nullptr)
@@ -50,38 +56,21 @@ void FModularVehicleSimulationCU::Simulate(UWorld* InWorld, float DeltaSeconds,
 		}
 	}
 
-	//if (InWorld)
-	//{
-	//	WriteNetReport(InWorld->IsNetMode(NM_Client), FString::Printf(TEXT("Frame %d,  Throttle %3.2f,  Brake %3.2f, Current Frame %d")
-	//		, InputData.PhysicsInputs.NetworkInputs.LocalFrame
-	//		, InputData.PhysicsInputs.NetworkInputs.VehicleInputs.Throttle
-	//		, InputData.PhysicsInputs.NetworkInputs.VehicleInputs.Brake
-	//		, CurrentFrame));
-	//}
-
-	check(Proxy->GetType() == EPhysicsProxyType::ClusterUnionProxy)
-	Simulate_ClusterUnion(InWorld, DeltaSeconds, InputData, OutputData, static_cast<Chaos::FClusterUnionPhysicsProxy*>(Proxy));
+	SimulateModuleTree(InWorld, DeltaSeconds, InputData, OutputData, Proxy);
 }
 
-void FModularVehicleSimulationCU::OnContactModification(Chaos::FCollisionContactModifier& Modifier, IPhysicsProxyBase* Proxy)
+void FModularVehicleSimulation::OnContactModification(Chaos::FCollisionContactModifier& Modifier, IPhysicsProxyBase* Proxy)
 {
 	using namespace Chaos;
 	Chaos::EnsureIsInPhysicsThreadContext();
-	check(Proxy->GetType() == EPhysicsProxyType::ClusterUnionProxy)
-	Chaos::FClusterUnionPhysicsProxy* ClusterProxy = static_cast<Chaos::FClusterUnionPhysicsProxy*>(Proxy);
-	check(ClusterProxy);
-
-	FPBDRigidsEvolutionGBF& Evolution = *static_cast<FPBDRigidsSolver*>(Proxy->GetSolver<FPBDRigidsSolver>())->GetEvolution();
-	FClusterUnionManager& ClusterUnionManager = Evolution.GetRigidClustering().GetClusterUnionManager();
-	const FClusterUnionIndex& CUI = ClusterProxy->GetClusterUnionIndex();
 
 	if (SimModuleTree.IsValid())
 	{
-		SimModuleTree->OnContactModification(Modifier, ClusterProxy);
+		SimModuleTree->OnContactModification(Modifier, Proxy);
 	}
 }
 
-void FModularVehicleSimulationCU::Simulate_ClusterUnion(UWorld* InWorld, float DeltaSeconds, const FModularVehicleAsyncInput& InputData, FModularVehicleAsyncOutput& OutputData, Chaos::FClusterUnionPhysicsProxy* Proxy)
+void FModularVehicleSimulation::SimulateModuleTree(UWorld* InWorld, float DeltaSeconds, const FModularVehicleAsyncInput& InputData, FModularVehicleAsyncOutput& OutputData, IPhysicsProxyBase* Proxy)
 {
 	Chaos::EnsureIsInPhysicsThreadContext();
 
@@ -109,181 +98,169 @@ void FModularVehicleSimulationCU::Simulate_ClusterUnion(UWorld* InWorld, float D
 		PerformAdditionalSimWork(InWorld, InputData, Proxy, SimInputData);
 
 		// run the dynamics simulation, engine, suspension, wheels, aerofoils etc.
-		SimModuleTree->Simulate(DeltaSeconds, SimInputData, Proxy);
+		SimModuleTree->Simulate(DeltaSeconds, SimInputData, Proxy, RootParticle);
 	}
 
 }
 
+void FModularVehicleSimulation::CacheRootParticle(IPhysicsProxyBase* Proxy)
+{
+	Chaos::EnsureIsInPhysicsThreadContext();
+	using namespace Chaos;
+	RootParticle = nullptr;
+
+	if (Proxy == nullptr)
+	{
+		return;
+	}
+
+	switch (Proxy->GetType())
+	{
+		case EPhysicsProxyType::ClusterUnionProxy:
+		{
+			if (FClusterUnionPhysicsProxy* CUProxy = static_cast<FClusterUnionPhysicsProxy*>(Proxy))
+			{
+				FPBDRigidsEvolutionGBF& Evolution = *static_cast<FPBDRigidsSolver*>(CUProxy->GetSolver<FPBDRigidsSolver>())->GetEvolution();
+				FClusterUnionManager& ClusterUnionManager = Evolution.GetRigidClustering().GetClusterUnionManager();
+				const FClusterUnionIndex& CUI = CUProxy->GetClusterUnionIndex();
+				if (FClusterUnion* ClusterUnion = ClusterUnionManager.FindClusterUnion(CUI))
+				{
+					if (FPBDRigidClusteredParticleHandle* ClusterHandle = ClusterUnion->InternalCluster)
+					{
+						RootParticle = ClusterHandle;
+					}
+				}
+			}
+		}
+		break;
+
+		case EPhysicsProxyType::SingleParticleProxy:
+		{
+			if (FSingleParticlePhysicsProxy* ParticleProxy = static_cast<FSingleParticlePhysicsProxy*>(Proxy))
+			{
+				RootParticle = ParticleProxy->GetHandle_LowLevel() ? ParticleProxy->GetHandle_LowLevel()->CastToRigidParticle() : nullptr;
+			}
+		}
+		break;
+
+		default:
+		{
+			UE_LOG(LogModularVehicleSim, Error, TEXT("Unsupported Particle type"));
+		}
+		break;
+	}
+}
 
-void FModularVehicleSimulationCU::PerformAdditionalSimWork(UWorld* InWorld, const FModularVehicleAsyncInput& InputData, Chaos::FClusterUnionPhysicsProxy* Proxy, Chaos::FAllInputs& AllInputs)
+void FModularVehicleSimulation::PerformAdditionalSimWork(UWorld* InWorld, const FModularVehicleAsyncInput& InputData, IPhysicsProxyBase* Proxy, Chaos::FAllInputs& AllInputs)
 {
 	using namespace Chaos;
 	check(Proxy);
 	Chaos::EnsureIsInPhysicsThreadContext();
 
-	if (SimModuleTree)
+	if (SimModuleTree && RootParticle)
 	{
-		FPBDRigidsEvolutionGBF& Evolution = *static_cast<FPBDRigidsSolver*>(Proxy->GetSolver<FPBDRigidsSolver>())->GetEvolution();
-		FClusterUnionManager& ClusterUnionManager = Evolution.GetRigidClustering().GetClusterUnionManager();
-		const FClusterUnionIndex& CUI = Proxy->GetClusterUnionIndex();
-		if (FClusterUnion* ClusterUnion = ClusterUnionManager.FindClusterUnion(CUI))
+		FRigidTransform3 ClusterWorldTM = FRigidTransform3(RootParticle->GetX(), RootParticle->GetR());
+
+		const TArray<FSimModuleTree::FSimModuleNode>& ModuleArray = SimModuleTree->GetSimulationModuleTree();
+
+		for (const FSimModuleTree::FSimModuleNode& Node : ModuleArray)
 		{
-			if (FPBDRigidClusteredParticleHandle* ClusterHandle = ClusterUnion->InternalCluster)
+			if (Node.IsValid() && Node.SimModule && Node.SimModule->IsEnabled())
 			{
-				TArray<FPBDRigidParticleHandle*> Particles = ClusterUnion->ChildParticles;
-				FRigidTransform3 ClusterWorldTM = FRigidTransform3(ClusterHandle->GetX(), ClusterHandle->GetR());
-
-				//Chaos::FDebugDrawQueue::GetInstance().DrawDebugCoordinateSystem(ClusterHandle->GetX(), FRotator(ClusterHandle->R()), 200, false, -1.f, 1.f, 3.0f);
+				FRigidTransform3 Frame = Node.SimModule->GetParentRelativeTransform();
 
-				const TArray<Chaos::FSimModuleTree::FSimModuleNode>& ModuleArray = SimModuleTree->GetSimulationModuleTree();
+						
+				AllInputs.VehicleWorldTransform = ClusterWorldTM;
 
-				for (const Chaos::FSimModuleTree::FSimModuleNode& Node : ModuleArray)
+				if (Node.SimModule->IsClustered() && Node.SimModule->IsBehaviourType(Chaos::eSimModuleTypeFlags::Raycast))
 				{
-					if (Node.IsValid() && Node.SimModule && Node.SimModule->IsEnabled())
-					{
-						FRigidTransform3 Frame = FRigidTransform3::Identity;
+					Chaos::FSpringTrace OutTrace;
+					Chaos::FSuspensionBaseInterface* Suspension = static_cast<Chaos::FSuspensionBaseInterface*>(Node.SimModule);
 
-						FPBDRigidParticleHandle* Child = Node.SimModule->GetParticleFromUniqueIndex(Node.SimModule->GetParticleIndex().Idx, Particles);
-						if (Child == nullptr)
+					// would be cleaner an faster to just store radius in suspension also
+					float WheelRadius = 0;
+					if (Suspension->GetWheelSimTreeIndex() != Chaos::ISimulationModuleBase::INVALID_IDX)
+					{
+						Chaos::FWheelBaseInterface* Wheel = static_cast<Chaos::FWheelBaseInterface*>(ModuleArray[Suspension->GetWheelSimTreeIndex()].SimModule);
+						if (Wheel)
 						{
-							continue;
+							WheelRadius = Wheel->GetWheelRadius();
 						}
+					}
 
-						if (FPBDRigidClusteredParticleHandle* ClusterChild = Child->CastToClustered(); ClusterChild && ClusterChild->IsChildToParentLocked())
-						{
-							Frame = ClusterChild->ChildToParent();
-						}
-						else
-						{
-							Frame = ClusterChild->ChildToParent();
+					Suspension->GetWorldRaycastLocation(ClusterWorldTM, WheelRadius, OutTrace);
 
-							const FRigidTransform3 ChildWorldTM(Child->GetX(), Child->GetR());
-							Frame = ChildWorldTM.GetRelativeTransform(ClusterWorldTM);
-						}
+					FVector TraceStart = OutTrace.Start;
+					FVector TraceEnd = OutTrace.End;
 
-						if (ClusterHandle)
-						{
-							AllInputs.VehicleWorldTransform = ClusterWorldTM;
+					const FCollisionQueryParams& TraceParams = InputData.PhysicsInputs.TraceParams;
+					FVector TraceVector(TraceStart - TraceEnd);
+					FVector TraceNormal = TraceVector.GetSafeNormal();
 
-							if (Node.SimModule->IsClustered() && Node.SimModule->IsBehaviourType(Chaos::eSimModuleTypeFlags::Raycast))
+					FHitResult HitResult = FHitResult();
+					ECollisionChannel SpringCollisionChannel = InputData.PhysicsInputs.CollisionChannel;
+					const FCollisionResponseParams& ResponseParams = InputData.PhysicsInputs.TraceCollisionResponse;
+					if (InWorld)
+					{
+						switch (InputData.PhysicsInputs.TraceType)
+						{
+							case ETraceType::Spherecast:
 							{
-								Chaos::FSpringTrace OutTrace;
-								Chaos::FSuspensionBaseInterface* Suspension = static_cast<Chaos::FSuspensionBaseInterface*>(Node.SimModule);
-
-								// would be cleaner an faster to just store radius in suspension also
-								float WheelRadius = 0;
-								if (Suspension->GetWheelSimTreeIndex() != Chaos::ISimulationModuleBase::INVALID_IDX)
-								{
-									Chaos::FWheelBaseInterface* Wheel = static_cast<Chaos::FWheelBaseInterface*>(ModuleArray[Suspension->GetWheelSimTreeIndex()].SimModule);
-									if (Wheel)
-									{
-										WheelRadius = Wheel->GetWheelRadius();
-									}
-								}
+								InWorld->SweepSingleByChannel(HitResult
+									, TraceStart + TraceNormal * WheelRadius
+									, TraceEnd + TraceNormal * WheelRadius
+									, FQuat::Identity, SpringCollisionChannel
+									, FCollisionShape::MakeSphere(WheelRadius), TraceParams
+									, ResponseParams);
+							}
+							break;
 
-								Suspension->GetWorldRaycastLocation(ClusterWorldTM, WheelRadius, OutTrace);
+							case ETraceType::Raycast:
+							default:
+							{
+								InWorld->LineTraceSingleByChannel(HitResult, TraceStart, TraceEnd, SpringCollisionChannel, TraceParams, ResponseParams);
+							}
+							break;
+						}
+					}
 
-								FVector TraceStart = OutTrace.Start;
-								FVector TraceEnd = OutTrace.End;
+					float Offset = Suspension->GetMaxSpringLength();
+					if (HitResult.bBlockingHit && GModularVehicleDebugParams.SuspensionRaycastsEnabled)
+					{
+						Offset = HitResult.Distance - WheelRadius;
 
-								const FCollisionQueryParams& TraceParams = InputData.PhysicsInputs.TraceParams;
-								FVector TraceVector(TraceStart - TraceEnd);
-								FVector TraceNormal = TraceVector.GetSafeNormal();
+						if (Suspension->GetWheelSimTreeIndex() != Chaos::ISimulationModuleBase::INVALID_IDX)
+						{
+							const Chaos::FSimModuleTree::FSimModuleNode& WheelNode = ModuleArray[Suspension->GetWheelSimTreeIndex()];
 
-								FHitResult HitResult = FHitResult();
-								ECollisionChannel SpringCollisionChannel = InputData.PhysicsInputs.CollisionChannel;
-								const FCollisionResponseParams& ResponseParams = InputData.PhysicsInputs.TraceCollisionResponse;
-								if (InWorld)
+							Chaos::FWheelBaseInterface* Wheel = static_cast<Chaos::FWheelBaseInterface*>(WheelNode.SimModule);
+							if (Wheel && HitResult.PhysMaterial.IsValid())
+							{
+								if (GModularVehicleDebugParams.FrictionOverride > 0)
 								{
-									switch (InputData.PhysicsInputs.TraceType)
-									{
-										case ETraceType::Spherecast:
-										{
-											InWorld->SweepSingleByChannel(HitResult
-												, TraceStart + TraceNormal * WheelRadius
-												, TraceEnd + TraceNormal * WheelRadius
-												, FQuat::Identity, SpringCollisionChannel
-												, FCollisionShape::MakeSphere(WheelRadius), TraceParams
-												, ResponseParams);
-										}
-										break;
-
-										case ETraceType::Raycast:
-										default:
-										{
-											InWorld->LineTraceSingleByChannel(HitResult, TraceStart, TraceEnd, SpringCollisionChannel, TraceParams, ResponseParams);
-										}
-										break;
-									}
+									Wheel->SetSurfaceFriction(GModularVehicleDebugParams.FrictionOverride);
 								}
-
-								float Offset = Suspension->GetMaxSpringLength();
-								if (HitResult.bBlockingHit && GModularVehicleDebugParams.SuspensionRaycastsEnabled)
+								else
 								{
-									Offset = HitResult.Distance - WheelRadius;
-
-									if (Suspension->GetWheelSimTreeIndex() != Chaos::ISimulationModuleBase::INVALID_IDX)
-									{
-										const Chaos::FSimModuleTree::FSimModuleNode& WheelNode = ModuleArray[Suspension->GetWheelSimTreeIndex()];
-
-										Chaos::FWheelBaseInterface* Wheel = static_cast<Chaos::FWheelBaseInterface*>(WheelNode.SimModule);
-										if (Wheel && HitResult.PhysMaterial.IsValid())
-										{
-											if (GModularVehicleDebugParams.FrictionOverride > 0)
-											{
-												Wheel->SetSurfaceFriction(GModularVehicleDebugParams.FrictionOverride);
-											}
-											else
-											{
-												Wheel->SetSurfaceFriction(HitResult.PhysMaterial->Friction);
-											}
-										}
-									}
-
-#if CHAOS_DEBUG_DRAW
-									if (GModularVehicleDebugParams.ShowSuspensionRaycasts)
-									{
-										Chaos::FDebugDrawQueue::GetInstance().DrawDebugSphere(HitResult.ImpactPoint, 3, 16, FColor::Red, false, -1.f, 0, 10.f);
-									}
-
-									if (Suspension->GetWheelSimTreeIndex() != Chaos::ISimulationModuleBase::INVALID_IDX)
-									{
-										Chaos::FWheelBaseInterface* Wheel = static_cast<Chaos::FWheelBaseInterface*>(ModuleArray[Suspension->GetWheelSimTreeIndex()].SimModule);
-										if (Wheel)
-										{
-											if (GModularVehicleDebugParams.ShowWheelData)
-											{
-												FString TextOut = FString::Format(TEXT("{0}"), { Wheel->GetForceIntoSurface() });
-												FColor Col = FColor::White;
-												if (InWorld)
-												{
-													if (InWorld->GetNetMode() == ENetMode::NM_Client)
-													{
-														Col = FColor::Blue;
-													}
-													else
-													{
-														Col = FColor::Red;
-													}
-												}
-												Chaos::FDebugDrawQueue::GetInstance().DrawDebugString(HitResult.ImpactPoint + FVec3(0, 50, 50), TextOut, nullptr, Col, -1.f, true, 1.0f);
-											}
-										}
-									}
-
-#endif
+									Wheel->SetSurfaceFriction(HitResult.PhysMaterial->Friction);
 								}
+							}
+						}
 
 #if CHAOS_DEBUG_DRAW
-								if (GModularVehicleDebugParams.ShowSuspensionRaycasts)
-								{
-									FColor DrawColor = FColor::Green;
-									DrawColor = (HitResult.bBlockingHit) ? FColor::Red : FColor::Green;
-									Chaos::FDebugDrawQueue::GetInstance().DrawDebugLine(TraceStart, TraceEnd, DrawColor, false, -1.f, 0, 2.f);
-									Chaos::FDebugDrawQueue::GetInstance().DrawDebugSphere(TraceStart, 3, 16, FColor::White, false, -1.f, 0, 10.f);
-									Chaos::FDebugDrawQueue::GetInstance().DrawDebugSphere(HitResult.ImpactPoint, 1, 16, FColor::Red, false, -1.f, 0, 10.f);
-									FString TextOut = FString::Format(TEXT("{0}"), { HitResult.Time });
+						if (GModularVehicleDebugParams.ShowSuspensionRaycasts)
+						{
+							Chaos::FDebugDrawQueue::GetInstance().DrawDebugSphere(HitResult.ImpactPoint, 3, 16, FColor::Red, false, -1.f, 0, 10.f);
+						}
 
+						if (Suspension->GetWheelSimTreeIndex() != Chaos::ISimulationModuleBase::INVALID_IDX)
+						{
+							Chaos::FWheelBaseInterface* Wheel = static_cast<Chaos::FWheelBaseInterface*>(ModuleArray[Suspension->GetWheelSimTreeIndex()].SimModule);
+							if (Wheel)
+							{
+								if (GModularVehicleDebugParams.ShowWheelData)
+								{
+									FString TextOut = FString::Format(TEXT("{0}"), { Wheel->GetForceIntoSurface() });
 									FColor Col = FColor::White;
 									if (InWorld)
 									{
@@ -298,93 +275,89 @@ void FModularVehicleSimulationCU::PerformAdditionalSimWork(UWorld* InWorld, cons
 									}
 									Chaos::FDebugDrawQueue::GetInstance().DrawDebugString(HitResult.ImpactPoint + FVec3(0, 50, 50), TextOut, nullptr, Col, -1.f, true, 1.0f);
 								}
-
-								if (GModularVehicleDebugParams.ShowRaycastMaterial)
-								{
-									if (HitResult.PhysMaterial.IsValid())
-									{
-										FDebugDrawQueue::GetInstance().DrawDebugString(HitResult.ImpactPoint, HitResult.PhysMaterial->GetName(), nullptr, FColor::White, -1.f, true, 1.0f);
-									}
-								}
-
-								if (GModularVehicleDebugParams.ShowWheelCollisionNormal)
-								{
-									FVector Pt = HitResult.ImpactPoint;
-									FDebugDrawQueue::GetInstance().DrawDebugLine(Pt, Pt + HitResult.Normal * 20.0f, FColor::Yellow, false, 1.0f, 0, 1.0f);
-									FDebugDrawQueue::GetInstance().DrawDebugSphere(Pt, 5.0f, 4, FColor::White, false, 1.0f, 0, 1.0f);
-								}
+							}
+						}
 
 #endif
-								Suspension->SetSpringLength(Offset, WheelRadius);
-								FVector Up = ClusterWorldTM.GetUnitAxis(EAxis::Z);
-
-								FVector HitPoint;
-								float HitDistance = 0.f;
-								if (InputData.PhysicsInputs.TraceType == ETraceType::Spherecast)
-								{
-									HitPoint = HitResult.Location;
-									HitDistance = HitResult.Distance;
-								}
-								else
-								{
-									HitPoint = HitResult.ImpactPoint + Up * WheelRadius;
-									HitDistance = HitResult.Distance - WheelRadius;
-								}
-
-								const TEnumAsByte<EPhysicalSurface> DefaultSurfaceType = EPhysicalSurface::SurfaceType_Default;
-								FSuspensionTargetPoint TargetPoint(
-									HitPoint
-									, HitResult.ImpactNormal
-									, HitDistance
-									, HitResult.bBlockingHit
-									, HitResult.PhysMaterial.IsValid() ? HitResult.PhysMaterial->SurfaceType : DefaultSurfaceType
-								);
+					}
 
-								Suspension->SetTargetPoint(TargetPoint);
+#if CHAOS_DEBUG_DRAW
+					if (GModularVehicleDebugParams.ShowSuspensionRaycasts)
+					{
+						FColor DrawColor = FColor::Green;
+						DrawColor = (HitResult.bBlockingHit) ? FColor::Red : FColor::Green;
+						Chaos::FDebugDrawQueue::GetInstance().DrawDebugLine(TraceStart, TraceEnd, DrawColor, false, -1.f, 0, 2.f);
+						Chaos::FDebugDrawQueue::GetInstance().DrawDebugSphere(TraceStart, 3, 16, FColor::White, false, -1.f, 0, 10.f);
+						Chaos::FDebugDrawQueue::GetInstance().DrawDebugSphere(HitResult.ImpactPoint, 1, 16, FColor::Red, false, -1.f, 0, 10.f);
+						FString TextOut = FString::Format(TEXT("{0}"), { HitResult.Time });
+
+						FColor Col = FColor::White;
+						if (InWorld)
+						{
+							if (InWorld->GetNetMode() == ENetMode::NM_Client)
+							{
+								Col = FColor::Blue;
+							}
+							else
+							{
+								Col = FColor::Red;
 							}
-
 						}
+						Chaos::FDebugDrawQueue::GetInstance().DrawDebugString(HitResult.ImpactPoint + FVec3(0, 50, 50), TextOut, nullptr, Col, -1.f, true, 1.0f);
 					}
 
-				}
-			}
-		}
+					if (GModularVehicleDebugParams.ShowRaycastMaterial)
+					{
+						if (HitResult.PhysMaterial.IsValid())
+						{
+							FDebugDrawQueue::GetInstance().DrawDebugString(HitResult.ImpactPoint, HitResult.PhysMaterial->GetName(), nullptr, FColor::White, -1.f, true, 1.0f);
+						}
+					}
 
-	}
-}
+					if (GModularVehicleDebugParams.ShowWheelCollisionNormal)
+					{
+						FVector Pt = HitResult.ImpactPoint;
+						FDebugDrawQueue::GetInstance().DrawDebugLine(Pt, Pt + HitResult.Normal * 20.0f, FColor::Yellow, false, 1.0f, 0, 1.0f);
+						FDebugDrawQueue::GetInstance().DrawDebugSphere(Pt, 5.0f, 4, FColor::White, false, 1.0f, 0, 1.0f);
+					}
 
-/**
- * ApplyDeferredForces should be called after the ParallelUpdatePT to send the calculated forces to the physics thread serially
- * as this cannot be done in parallel
- */
-void FModularVehicleSimulationCU::ApplyDeferredForces(FGeometryCollectionPhysicsProxy* Proxy)
-{
-	Chaos::EnsureIsInPhysicsThreadContext();
+#endif
+					Suspension->SetSpringLength(Offset, WheelRadius);
+					FVector Up = ClusterWorldTM.GetUnitAxis(EAxis::Z);
 
-	if (SimModuleTree && Proxy)
-	{
-		if (UGeometryCollectionComponent* GCComponent = Cast<UGeometryCollectionComponent>(Proxy->GetOwner()))
-		{
-			check(GCComponent->GetOwner());
+					FVector HitPoint;
+					float HitDistance = 0.f;
+					if (InputData.PhysicsInputs.TraceType == ETraceType::Spherecast)
+					{
+						HitPoint = HitResult.Location;
+						HitDistance = HitResult.Distance;
+					}
+					else
+					{
+						HitPoint = HitResult.ImpactPoint + Up * WheelRadius;
+						HitDistance = HitResult.Distance - WheelRadius;
+					}
 
-			if (const UGeometryCollection* Rest = GCComponent->GetRestCollection())
-			{
-				if (Rest->GetGeometryCollection() && Rest->GetGeometryCollection()->HasAttribute(TEXT("MassToLocal"), FTransformCollection::TransformGroup))
-				{
-					const TManagedArray<FTransform>& CollectionMassToLocal = Rest->GetGeometryCollection()->GetAttribute<FTransform>(TEXT("MassToLocal"), FTransformCollection::TransformGroup);
+					const TEnumAsByte<EPhysicalSurface> DefaultSurfaceType = EPhysicalSurface::SurfaceType_Default;
+					FSuspensionTargetPoint TargetPoint(
+						HitPoint
+						, HitResult.ImpactNormal
+						, HitDistance
+						, HitResult.bBlockingHit
+						, HitResult.PhysMaterial.IsValid() ? HitResult.PhysMaterial->SurfaceType : DefaultSurfaceType
+					);
 
-					SimModuleTree->AccessDeferredForces().Apply(
-						Proxy,
-						Rest->GetGeometryCollection()->Transform,
-						CollectionMassToLocal,
-						Rest->GetGeometryCollection()->Parent);
+					Suspension->SetTargetPoint(TargetPoint);
 				}
+
 			}
+
 		}
+
 	}
 }
 
-void FModularVehicleSimulationCU::ApplyDeferredForces(Chaos::FClusterUnionPhysicsProxy* Proxy)
+void FModularVehicleSimulation::ApplyDeferredForces(IPhysicsProxyBase* Proxy)
 {
 	using namespace Chaos;
 
@@ -392,26 +365,13 @@ void FModularVehicleSimulationCU::ApplyDeferredForces(Chaos::FClusterUnionPhysic
 
 	if (SimModuleTree && Proxy)
 	{
-		// This gives us access to the PT parent cluster and child particles
-		FPBDRigidsEvolutionGBF& Evolution = *static_cast<FPBDRigidsSolver*>(Proxy->GetSolver<FPBDRigidsSolver>())->GetEvolution();
-		FClusterUnionManager& ClusterUnionManager = Evolution.GetRigidClustering().GetClusterUnionManager();
-		const FClusterUnionIndex& CUI = Proxy->GetClusterUnionIndex();
-		if (FClusterUnion* ClusterUnion = ClusterUnionManager.FindClusterUnion(CUI))
-		{
-			FPBDRigidClusteredParticleHandle* ClusterHandle = ClusterUnion->InternalCluster;
-			TArray<FPBDRigidParticleHandle*> Particles = ClusterUnion->ChildParticles;
 
-			TArray<FPBDRigidClusteredParticleHandle*> Clusters;
-			Clusters.Add(ClusterHandle);
-
-			SimModuleTree->AccessDeferredForces().Apply(
-				Particles,
-				Clusters);
-		}
+		SimModuleTree->AccessDeferredForces().Apply(RootParticle);
+		
 	}
 }
 
-void FModularVehicleSimulationCU::FillOutputState(FModularVehicleAsyncOutput& Output)
+void FModularVehicleSimulation::FillOutputState(FModularVehicleAsyncOutput& Output)
 {
 	if (Chaos::FSimModuleTree* SimTree = GetSimComponentTree().Get())
 	{
diff --git a/Engine/Plugins/Experimental/ChaosModularVehicle/Source/ChaosModularVehicleEngine/Private/ChaosModularVehicle/SuspensionSimModule.cpp b/Engine/Plugins/Experimental/ChaosModularVehicle/Source/ChaosModularVehicleEngine/Private/ChaosModularVehicle/SuspensionSimModule.cpp
--- a/Engine/Plugins/Experimental/ChaosModularVehicle/Source/ChaosModularVehicleEngine/Private/ChaosModularVehicle/SuspensionSimModule.cpp
+++ b/Engine/Plugins/Experimental/ChaosModularVehicle/Source/ChaosModularVehicleEngine/Private/ChaosModularVehicle/SuspensionSimModule.cpp
@@ -54,10 +54,10 @@ using namespace Chaos;
 		float TraceLength = OutTrace.Start.Z - OutTrace.End.Z;
 	}
 
-	void FSuspensionSimModule::OnConstruction_External(Chaos::FClusterUnionPhysicsProxy* Proxy)
+	void FSuspensionSimModule::OnConstruction_External(const Chaos::FPhysicsObjectHandle& PhysicsObject)
 	{
 		Chaos::EnsureIsInGameThreadContext();
-		CreateConstraint(Proxy);
+		CreateConstraint(PhysicsObject);
 	}
 
 	void FSuspensionSimModule::OnTermination_External()
@@ -112,49 +112,39 @@ using namespace Chaos;
 
 	}
 
-	void FSuspensionSimModule::Animate(Chaos::FClusterUnionPhysicsProxy* Proxy)
+	void FSuspensionSimModule::Animate()
 	{
-		if (FPBDRigidClusteredParticleHandle* ClusterChild = GetClusterParticle(Proxy))
-		{
-			float CurrentSpringLength = GetSpringLength();
-
-			FVector RestPos = GetInitialParticleTransform().GetTranslation();
+		FVector Movement = -Setup().SuspensionAxis * (Setup().MaxRaise + GetSpringLength());
 
-			FVector Movement = Setup().SuspensionAxis * (Setup().MaxRaise + CurrentSpringLength);
-			Movement = GetComponentTransform().TransformVector(Movement);
-			FVector NewPos = RestPos - Movement;
-			ClusterChild->ChildToParent().SetTranslation(NewPos); // local frame for module
-		}
+		AnimationData.AnimFlags = EAnimationFlags::AnimatePosition;
+		AnimationData.AnimationLocOffset = Movement;
 	}
 
-	void FSuspensionSimModule::CreateConstraint(Chaos::FClusterUnionPhysicsProxy* Proxy)
+
+	void FSuspensionSimModule::CreateConstraint(const Chaos::FPhysicsObjectHandle& PhysicsObject)
 	{
 		Chaos::EnsureIsInGameThreadContext();
+
 		const FVector& LocalOffset = GetInitialParticleTransform().GetLocation();
 
-		if (Proxy)
+		if (FChaosScene* Scene = static_cast<FChaosScene*>(FPhysicsObjectExternalInterface::GetScene({ &PhysicsObject, 1 })))
 		{
-			const FPhysicsObjectHandle& PhysicsObject = Proxy->GetPhysicsObjectHandle();
-
-			if (FChaosScene* Scene = static_cast<FChaosScene*>(FPhysicsObjectExternalInterface::GetScene({ &PhysicsObject, 1 })))
+			FLockedWritePhysicsObjectExternalInterface Interface = FPhysicsObjectExternalInterface::LockWrite(Scene);
+			if (const FGeometryParticle* Particle = Interface->GetParticle(PhysicsObject))
 			{
-				FLockedWritePhysicsObjectExternalInterface Interface = FPhysicsObjectExternalInterface::LockWrite(Scene);
-				if (const FGeometryParticle* Particle = Interface->GetParticle(PhysicsObject))
-				{
-					ConstraintHandle = FPhysicsInterface::CreateSuspension(PhysicsObject, LocalOffset);
+				ConstraintHandle = FPhysicsInterface::CreateSuspension(PhysicsObject, LocalOffset);
 
-					if (ConstraintHandle.IsValid())
+				if (ConstraintHandle.IsValid())
+				{
+					if (FSuspensionConstraint* Constraint = static_cast<FSuspensionConstraint*>(ConstraintHandle.Constraint))
 					{
-						if (FSuspensionConstraint* Constraint = static_cast<FSuspensionConstraint*>(ConstraintHandle.Constraint))
-						{
-							Constraint->SetHardstopStiffness(1.0f);
-							Constraint->SetSpringStiffness(Setup().SpringRate * 0.25f);
-							Constraint->SetSpringPreload(Setup().SpringPreload);
-							Constraint->SetSpringDamping(Setup().SpringDamping * 5.0f);
-							Constraint->SetMinLength(-Setup().MaxRaise);
-							Constraint->SetMaxLength(Setup().MaxDrop);
-							Constraint->SetAxis(-Setup().SuspensionAxis);
-						}
+						Constraint->SetHardstopStiffness(1.0f);
+						Constraint->SetSpringStiffness(Setup().SpringRate * 0.25f);
+						Constraint->SetSpringPreload(Setup().SpringPreload);
+						Constraint->SetSpringDamping(Setup().SpringDamping * 5.0f);
+						Constraint->SetMinLength(-Setup().MaxRaise);
+						Constraint->SetMaxLength(Setup().MaxDrop);
+						Constraint->SetAxis(-Setup().SuspensionAxis);
 					}
 				}
 			}
@@ -233,10 +223,6 @@ using namespace Chaos;
 			SpringDisplacementVector = -Sim->Setup().SuspensionAxis * Sim->SpringDisplacement + Sim->GetAnimationOffset();
 			SpringSpeed = Sim->SpringSpeed;
 			ImpactNormal = Sim->GetImpactNormal();
-
-			AnimFlags = EAnimationFlags::AnimatePosition;
-			AnimationLocOffset = SpringDisplacementVector;
-
 		}
 	}
 
diff --git a/Engine/Plugins/Experimental/ChaosModularVehicle/Source/ChaosModularVehicleEngine/Public/ChaosModularVehicle/ClusterUnionVehicleComponent.h b/Engine/Plugins/Experimental/ChaosModularVehicle/Source/ChaosModularVehicleEngine/Public/ChaosModularVehicle/ClusterUnionVehicleComponent.h
--- a/Engine/Plugins/Experimental/ChaosModularVehicle/Source/ChaosModularVehicleEngine/Public/ChaosModularVehicle/ClusterUnionVehicleComponent.h
+++ b/Engine/Plugins/Experimental/ChaosModularVehicle/Source/ChaosModularVehicleEngine/Public/ChaosModularVehicle/ClusterUnionVehicleComponent.h
@@ -4,6 +4,7 @@
 
 #include "PhysicsEngine/ClusterUnionComponent.h"
 #include "ChaosModularVehicle/ModularVehicleSocket.h"
+#include "PhysicsProxy/ClusterUnionPhysicsProxy.h"
 
 #include "ClusterUnionVehicleComponent.generated.h"
 
diff --git a/Engine/Plugins/Experimental/ChaosModularVehicle/Source/ChaosModularVehicleEngine/Public/ChaosModularVehicle/ModularVehicleBaseComponent.h b/Engine/Plugins/Experimental/ChaosModularVehicle/Source/ChaosModularVehicleEngine/Public/ChaosModularVehicle/ModularVehicleBaseComponent.h
--- a/Engine/Plugins/Experimental/ChaosModularVehicle/Source/ChaosModularVehicleEngine/Public/ChaosModularVehicle/ModularVehicleBaseComponent.h
+++ b/Engine/Plugins/Experimental/ChaosModularVehicle/Source/ChaosModularVehicleEngine/Public/ChaosModularVehicle/ModularVehicleBaseComponent.h
@@ -64,11 +64,13 @@ struct CHAOSMODULARVEHICLEENGINE_API FModuleAnimationSetup
 {
 	GENERATED_USTRUCT_BODY()
 
-	FModuleAnimationSetup(FName BoneNameIn)
+	FModuleAnimationSetup(FName BoneNameIn, int TransformIndexIn)
 		: BoneName(BoneNameIn)
 		, RotOffset(FRotator::ZeroRotator)
 		, LocOffset(FVector::ZeroVector)
+		, CombinedRotation(FQuat::Identity)
 		, AnimFlags(0)
+		, TransformIndex(TransformIndexIn)
 	{
 	}
 
@@ -76,14 +78,18 @@ struct CHAOSMODULARVEHICLEENGINE_API FModuleAnimationSetup
 		: BoneName(NAME_None)
 		, RotOffset(FRotator::ZeroRotator)
 		, LocOffset(FVector::ZeroVector)
+		, CombinedRotation(FQuat::Identity)
 		, AnimFlags(0)
+		, TransformIndex(INDEX_NONE)
 	{
 	}
 
-	FName BoneName;
+	FName BoneName;		// required for skeletal mesh
 	FRotator RotOffset;
 	FVector LocOffset;
+	FQuat CombinedRotation;
 	uint16 AnimFlags;
+	int TransformIndex; // required for non skeletal mesh case
 };
 
 
@@ -125,7 +131,11 @@ public:
 
 	void ProduceInput(int32 PhysicsStep, int32 NumSteps);
 
-	void CreateAssociatedSimComponents(USceneComponent* AttachedComponent, int ParentIndex, int TransformIndex, Chaos::FSimTreeUpdates& TreeUpdatesOut);
+	// Create sim components hierarchy from parent component - used by original cluster code union path
+	virtual void CreateAssociatedSimComponents(USceneComponent* AttachedComponent, int ParentIndex, int TransformIndex, Chaos::FSimTreeUpdates& TreeUpdatesOut);
+
+	// Create sim components in hierarchy from parent component - used by code path that does not rely on cluster union
+	virtual void CreateIndependentSimComponents(USceneComponent* RootComponent, USceneComponent* AttachedComponent, int ParentIndex, int TransformIndex, Chaos::FSimTreeUpdates& TreeUpdatesOut);
 
 	void PreTickGT(float DeltaTime);
 	void UpdateState(float DeltaTime);
@@ -263,18 +273,32 @@ public:
 	TArray<FModuleAnimationSetup>& AccessModuleAnimationSetups() { return ModuleAnimationSetups; }
 	const TArray<FModuleAnimationSetup>& GetModuleAnimationSetups() const { return ModuleAnimationSetups; }
 
+	void SetPhysicsProxy(IPhysicsProxyBase* PhysicsProxy) { CachedPhysicsProxy = PhysicsProxy; }
+	void SetPhysicsProxyIfNotAlreadySpecified(IPhysicsProxyBase* PhysicsProxy)
+	{ 
+		if (CachedPhysicsProxy == nullptr) 
+		{
+			CachedPhysicsProxy = PhysicsProxy; 
+		}
+	}
+
 protected:
 
 	void CreateVehicleSim();
 	void DestroyVehicleSim();
 	void UpdatePhysicalProperties();
 	void AddOverlappingComponentsToCluster();
-	void AddGeometryCollectionsFromOwnedActor();
+	bool AddGeometryCollectionsFromOwnedActor();
 	void SetupSkeletalAnimationStructure();
 	void AssimilateComponentInputs(TArray<FModuleInputSetup>& OutCombinedInputs);
+	void CacheRootPhysicsObject(IPhysicsProxyBase* Proxy);
+
 	// #TODO reinstate? virtual void GenerateInputModifiers(const TArray<FModuleInputSetup>& CombinedInputConfiguration);
 	// #TODO reinstate? virtual void ApplyInputModifiers(float DeltaTime, const FModuleInputContainer& RawValue);
+	void AddSimulationComponentsFromOwnedActor(); // This version does not require cluster unions or geometry collections as parent components
 
+	virtual void AddComponentToSimulationImpl(UPrimitiveComponent* Component, const TArray<FClusterUnionBoneData>& BonesData, const TArray<FClusterUnionBoneData>& RemovedBoneIDs, bool bIsNew);
+	virtual void RemoveComponentFromSimulationImpl(UPrimitiveComponent* Component, const TArray<FClusterUnionBoneData>& RemovedBonesData);
 
 	void ActionTreeUpdates(Chaos::FSimTreeUpdates* NextTreeUpdates);
 
@@ -282,9 +306,9 @@ protected:
 	int32 FindParentsLastSimComponent(const USceneComponent* AttachedComponent);
 
 	IPhysicsProxyBase* GetPhysicsProxy() const;
+	int GenerateNewGuid();
 
-
-	int32 FindComponentAddOrder(UPrimitiveComponent* InComponent);
+	int32 FindComponentAddOrder(USceneComponent* InComponent);
 	bool FindAndRemoveNextPendingUpdate(int32 NextIndex, Chaos::FSimTreeUpdates* OutData);
 
 	// replicated state of vehicle 
@@ -327,6 +351,13 @@ public:
 
 	void LogInputSetup();
 
+	int AddSimModule(Chaos::ISimulationModuleBase* NewModule, const FTransform& ComponentTransform, int ParentIndex, int TransformIndex);
+	void RemoveSimModule(int ModuleGuid);
+	void FinalizeModuleUpdates();
+
+	void EnableModuleCreationFromActorComponents(bool bEnable) { bEnableModuleCreationFromActorComponents = bEnable; }
+	Chaos::FSimTreeUpdates StoredTreeUpdates;
+
 	TArray<AActor*> ActorsToIgnore;
 	EChaosAsyncVehicleDataType CurAsyncType;
 	FModularVehicleAsyncInput* CurAsyncInput;
@@ -348,7 +379,7 @@ public:
 
 	TArray<FAsyncOutputWrapper> OutputsWaitingOn;
 	TUniquePtr<FPhysicsVehicleOutput> PVehicleOutput;	/* physics simulation data output from the async physics thread */
-	TUniquePtr<FModularVehicleSimulationCU> VehicleSimulationPT;	/* simulation code running on the physics thread async callback */
+	TUniquePtr<FModularVehicleSimulation> VehicleSimulationPT;	/* simulation code running on the physics thread async callback */
 
 private:
 
@@ -359,13 +390,14 @@ private:
 	float PrevSteeringInput = 0.0f;
 
 	int32 LastComponentAddIndex = INDEX_NONE;
-	TMap<TObjectKey<UPrimitiveComponent>, Chaos::FSimTreeUpdates> PendingTreeUpdates;
+	TMap<TObjectKey<USceneComponent>, Chaos::FSimTreeUpdates> PendingTreeUpdates;
 
 	int32 NextConstructionIndex = 0;
 
 	int32 ClusteringCount = 0;
 
 	bool bIsLocallyControlled;
+	bool bEnableModuleCreationFromActorComponents = true;
 
 	TArray<FModuleAnimationSetup> ModuleAnimationSetups;
 
@@ -375,4 +407,7 @@ private:
 	TObjectPtr<UVehicleInputProducerBase> InputProducer = nullptr;
 
 	FModuleInputContainer InputsContainer;
+
+	Chaos::FPhysicsObjectHandle RootPhysicsObject = nullptr;
+	mutable IPhysicsProxyBase* CachedPhysicsProxy = nullptr;
 };
diff --git a/Engine/Plugins/Experimental/ChaosModularVehicle/Source/ChaosModularVehicleEngine/Public/ChaosModularVehicle/ModularVehicleSimulationCU.h b/Engine/Plugins/Experimental/ChaosModularVehicle/Source/ChaosModularVehicleEngine/Public/ChaosModularVehicle/ModularVehicleSimulationCU.h
--- a/Engine/Plugins/Experimental/ChaosModularVehicle/Source/ChaosModularVehicleEngine/Public/ChaosModularVehicle/ModularVehicleSimulationCU.h
+++ b/Engine/Plugins/Experimental/ChaosModularVehicle/Source/ChaosModularVehicleEngine/Public/ChaosModularVehicle/ModularVehicleSimulationCU.h
@@ -2,13 +2,16 @@
 
 #pragma once
 
-#include "GeometryCollection/GeometryCollectionComponent.h"
 #include "ChaosModularVehicle/ChaosSimModuleManagerAsyncCallback.h"
 #include "SimModule/SimModuleTree.h"
 #include "SimModule/SimModulesInclude.h"
 #include "ModularVehicleBuilder.h"
+#include "Chaos/ParticleHandle.h"
+#include "Chaos/ParticleHandleFwd.h"
 #include "Misc/TransactionallySafeRWLock.h"
 
+DECLARE_LOG_CATEGORY_EXTERN(LogModularVehicleSim, Log, All);
+
 struct FModularVehicleAsyncInput;
 struct FChaosSimModuleManagerAsyncOutput;
 struct FModuleInputContainer;
@@ -29,21 +32,22 @@ struct CHAOSMODULARVEHICLEENGINE_API FModularVehicleDebugParams
 namespace Chaos
 {
 	class FClusterUnionPhysicsProxy;
+	class FSingleParticlePhysicsProxy;
 }
 
 
-class CHAOSMODULARVEHICLEENGINE_API FModularVehicleSimulationCU
+class CHAOSMODULARVEHICLEENGINE_API FModularVehicleSimulation
 {
 public:
 	using FInputNameMap = FInputInterface::FInputNameMap;
 
-	FModularVehicleSimulationCU(bool InUsingNetworkPhysicsPrediction, int8 InNetMode)
+	FModularVehicleSimulation(bool InUsingNetworkPhysicsPrediction, int8 InNetMode)
 		: bUsingNetworkPhysicsPrediction(InUsingNetworkPhysicsPrediction)
 		, NetMode(InNetMode)
 	{
 	}
 
-	virtual ~FModularVehicleSimulationCU()
+	virtual ~FModularVehicleSimulation()
 	{
 		SimModuleTree.Reset();
 	}
@@ -55,17 +59,17 @@ public:
 		FTransactionallySafeWriteScopeLock InputConfigLock(InputConfigurationLock);
 		InputNameMap = InNameMap; 
 	}
+	void CacheRootParticle(IPhysicsProxyBase* Proxy);
 
 	/** Update called from Physics Thread */
 	virtual void Simulate(UWorld* InWorld, float DeltaSeconds, const FModularVehicleAsyncInput& InputData, FModularVehicleAsyncOutput& OutputData, IPhysicsProxyBase* Proxy);
-	virtual void Simulate_ClusterUnion(UWorld* InWorld, float DeltaSeconds, const FModularVehicleAsyncInput& InputData, FModularVehicleAsyncOutput& OutputData, Chaos::FClusterUnionPhysicsProxy* Proxy);
+	virtual void SimulateModuleTree(UWorld* InWorld, float DeltaSeconds, const FModularVehicleAsyncInput& InputData, FModularVehicleAsyncOutput& OutputData, IPhysicsProxyBase* Proxy);
 
 	virtual void OnContactModification(Chaos::FCollisionContactModifier& Modifier, IPhysicsProxyBase* Proxy);
 
-	void ApplyDeferredForces(FGeometryCollectionPhysicsProxy* RigidHandle);
-	void ApplyDeferredForces(Chaos::FClusterUnionPhysicsProxy* Proxy);
+	void ApplyDeferredForces(IPhysicsProxyBase* Proxy);
 
-	void PerformAdditionalSimWork(UWorld* InWorld, const FModularVehicleAsyncInput& InputData, Chaos::FClusterUnionPhysicsProxy* Proxy, Chaos::FAllInputs& AllInputs);
+	void PerformAdditionalSimWork(UWorld* InWorld, const FModularVehicleAsyncInput& InputData, IPhysicsProxyBase* Proxy, Chaos::FAllInputs& AllInputs);
 
 	void FillOutputState(FModularVehicleAsyncOutput& Output);
 
@@ -89,4 +93,7 @@ public:
 
 	int8 NetMode;
 
+	Chaos::FPBDRigidParticleHandle* RootParticle = nullptr; // cached root particle
+
 };
+
diff --git a/Engine/Plugins/Experimental/ChaosModularVehicle/Source/ChaosModularVehicleEngine/Public/ChaosModularVehicle/SuspensionSimModule.h b/Engine/Plugins/Experimental/ChaosModularVehicle/Source/ChaosModularVehicleEngine/Public/ChaosModularVehicle/SuspensionSimModule.h
--- a/Engine/Plugins/Experimental/ChaosModularVehicle/Source/ChaosModularVehicleEngine/Public/ChaosModularVehicle/SuspensionSimModule.h
+++ b/Engine/Plugins/Experimental/ChaosModularVehicle/Source/ChaosModularVehicleEngine/Public/ChaosModularVehicle/SuspensionSimModule.h
@@ -144,19 +144,19 @@ public:
 	virtual  void SetSpringLength(float InLength, float WheelRadius) override;
 	virtual void GetWorldRaycastLocation(const FTransform& BodyTransform, float WheelRadius, Chaos::FSpringTrace& OutTrace) override;
 
-	virtual void OnConstruction_External(Chaos::FClusterUnionPhysicsProxy* Proxy) override;
+	virtual void OnConstruction_External(const Chaos::FPhysicsObjectHandle& PhysicsObject) override;
 	virtual void OnTermination_External() override;
 
 	virtual void Simulate(float DeltaTime, const Chaos::FAllInputs& Inputs, Chaos::FSimModuleTree& VehicleModuleSystem) override;
 
-	virtual void Animate(Chaos::FClusterUnionPhysicsProxy* Proxy) override;
+	virtual void Animate() override;
 
 	const FVector& GetRestLocation() const { return Setup().RestOffset; }
 
 	void UpdateConstraint();
 
 protected:
-	void CreateConstraint(Chaos::FClusterUnionPhysicsProxy* Proxy);
+	void CreateConstraint(const Chaos::FPhysicsObjectHandle& PhysicsObject);
 	void DestroyConstraint();
 
 private:
diff --git a/Engine/Source/Runtime/Experimental/ChaosVehicles/ChaosVehiclesCore/Private/SimModule/AerofoilModule.cpp b/Engine/Source/Runtime/Experimental/ChaosVehicles/ChaosVehiclesCore/Private/SimModule/AerofoilModule.cpp
--- a/Engine/Source/Runtime/Experimental/ChaosVehicles/ChaosVehiclesCore/Private/SimModule/AerofoilModule.cpp
+++ b/Engine/Source/Runtime/Experimental/ChaosVehicles/ChaosVehiclesCore/Private/SimModule/AerofoilModule.cpp
@@ -146,16 +146,22 @@ namespace Chaos
 		return Force;
 	}
 
-	void FAerofoilSimModule::Animate(Chaos::FClusterUnionPhysicsProxy* Proxy)
+	void FAerofoilSimModule::Animate()
 	{
-		if (FPBDRigidClusteredParticleHandle* ClusterChild = GetClusterParticle(Proxy))
+		AnimationData.AnimFlags = EAnimationFlags::AnimateRotation;
+		AnimationData.CombinedRotation = FQuat(Setup().ControlRotationAxis, FMath::DegreesToRadians(ControlSurfaceAngle) * Setup().AnimationMagnitudeMultiplier);
+
+		switch (Setup().Type)
 		{
-			FQuat ControlRotation = FQuat(Setup().ControlRotationAxis, FMath::DegreesToRadians(ControlSurfaceAngle * Setup().AnimationMagnitudeMultiplier));
+		case EAerofoil::Wing:
+		case EAerofoil::Elevator:
+			AnimationData.AnimationRotOffset.Pitch = ControlSurfaceAngle * Setup().AnimationMagnitudeMultiplier;
+			break;
 
-			FTransform InitialTransform = GetInitialParticleTransform();
-			ClusterChild->ChildToParent().SetRotation(InitialTransform.GetRotation() * ControlRotation);
+		case EAerofoil::Rudder:
+			AnimationData.AnimationRotOffset.Yaw = ControlSurfaceAngle * Setup().AnimationMagnitudeMultiplier;
+			break;
 		}
-
 	}
 
 
@@ -164,24 +170,6 @@ namespace Chaos
 		check(SimModule->IsSimType<class FAerofoilSimModule>());
 
 		FSimOutputData::FillOutputState(SimModule);
-
-		if (const FAerofoilSimModule* Sim = static_cast<const FAerofoilSimModule*>(SimModule))
-		{
-			AnimFlags = EAnimationFlags::AnimateRotation;
-
-			switch (Sim->Setup().Type)
-			{
-			case EAerofoil::Wing:
-			case EAerofoil::Elevator:
-				AnimationRotOffset.Pitch = Sim->ControlSurfaceAngle * Sim->Setup().AnimationMagnitudeMultiplier;
-				break;
-
-			case EAerofoil::Rudder:
-				AnimationRotOffset.Yaw = Sim->ControlSurfaceAngle * Sim->Setup().AnimationMagnitudeMultiplier;
-				break;
-			}
-
-		}
 	}
 
 	void FAerofoilOutputData::Lerp(const FSimOutputData& InCurrent, const FSimOutputData& InNext, float Alpha)
diff --git a/Engine/Source/Runtime/Experimental/ChaosVehicles/ChaosVehiclesCore/Private/SimModule/DeferredForcesModular.cpp b/Engine/Source/Runtime/Experimental/ChaosVehicles/ChaosVehiclesCore/Private/SimModule/DeferredForcesModular.cpp
--- a/Engine/Source/Runtime/Experimental/ChaosVehicles/ChaosVehiclesCore/Private/SimModule/DeferredForcesModular.cpp
+++ b/Engine/Source/Runtime/Experimental/ChaosVehicles/ChaosVehiclesCore/Private/SimModule/DeferredForcesModular.cpp
@@ -333,7 +333,50 @@ void FDeferredForcesModular::Apply(TArray<Chaos::FPBDRigidParticleHandle*>& Part
 	ApplyTorqueDatas.Empty();
 }
 
+void FDeferredForcesModular::Apply(Chaos::FPBDRigidParticleHandle* RigidHandle)
+{
+	if (!GCoreModularVehicleDebugParams.DisableForces)
+	{
+		FTransform RelativeTransform;
+
+		for (const FApplyForceData& Data : ApplyForceDatas)
+		{
+			if (RigidHandle)
+			{
+				AddForce(RigidHandle, Data, RelativeTransform);
+			}
+		}
+
+		for (const FApplyForceAtPositionData& Data : ApplyForceAtCOMDatas)
+		{
+			if (RigidHandle)
+			{
+				AddForceAtCOM(RigidHandle, Data);
+			}
+		}
+
+		for (const FApplyForceAtPositionData& Data : ApplyForceAtPositionDatas)
+		{
+			if (RigidHandle)
+			{
+				AddForceAtPosition(RigidHandle, Data, RelativeTransform);
+			}
+		}
 
+		for (const FAddTorqueInRadiansData& Data : ApplyTorqueDatas)
+		{
+			if (RigidHandle)
+			{
+				AddTorque(RigidHandle, Data, RelativeTransform);
+			}
+		}
+	}
+
+	ApplyForceDatas.Empty();
+	ApplyForceAtCOMDatas.Empty();
+	ApplyForceAtPositionDatas.Empty();
+	ApplyTorqueDatas.Empty();
+}
 
 void AddForce_Implementation(Chaos::FPBDRigidParticleHandle* RigidHandle, const FTransform& ParticleOffsetTransform, const FVector& LocalForce, const FTransform& OffsetTransform, bool bLevelSlope, const FColor& DebugColor)
 {
diff --git a/Engine/Source/Runtime/Experimental/ChaosVehicles/ChaosVehiclesCore/Private/SimModule/SimModuleTree.cpp b/Engine/Source/Runtime/Experimental/ChaosVehicles/ChaosVehiclesCore/Private/SimModule/SimModuleTree.cpp
--- a/Engine/Source/Runtime/Experimental/ChaosVehicles/ChaosVehiclesCore/Private/SimModule/SimModuleTree.cpp
+++ b/Engine/Source/Runtime/Experimental/ChaosVehicles/ChaosVehiclesCore/Private/SimModule/SimModuleTree.cpp
@@ -5,6 +5,7 @@
 #include "Chaos/GeometryParticlesfwd.h"
 #include "PhysicsProxy/GeometryCollectionPhysicsProxy.h"
 #include "PhysicsProxy/ClusterUnionPhysicsProxy.h"
+#include "PhysicsProxy/SingleParticlePhysicsProxy.h"
 #include "Chaos/PhysicsObjectInternalInterface.h"
 #include "Chaos/DebugDrawQueue.h"
 #include "VehicleUtility.h"
@@ -220,17 +221,17 @@ void FSimModuleTree::DeleteNode(int AtIndex)
 }
 
 
-void FSimModuleTree::Simulate(float DeltaTime, FAllInputs& Inputs, FClusterUnionPhysicsProxy* PhysicsProxy)
+void FSimModuleTree::Simulate(float DeltaTime, FAllInputs& Inputs, IPhysicsProxyBase* PhysicsProxy, Chaos::FPBDRigidParticleHandle* RootParticle)
 {
 	SCOPE_CYCLE_COUNTER(STAT_ModularVehicle_SimulateTree);
 
 	if (IsSimulationEnabled())
 	{
-		if (PhysicsProxy)
+		if (PhysicsProxy && RootParticle)
 		{
-			UpdateVehicleState(PhysicsProxy);
+			UpdateVehicleState(RootParticle);
 
-			UpdateModuleVelocites(PhysicsProxy, Inputs.GetControls().InputsNonZero() || Inputs.bKeepVehicleAwake);
+			UpdateModuleVelocites(PhysicsProxy, RootParticle, Inputs.GetControls().InputsNonZero() || Inputs.bKeepVehicleAwake);
 		}
 
 		TArray<int> RootNodes;
@@ -238,13 +239,13 @@ void FSimModuleTree::Simulate(float DeltaTime, FAllInputs& Inputs, FClusterUnion
 
 		if (SimTreeProcessingOrder == ESimTreeProcessingOrder::LeafFirstBFS)
 		{
-			SimulateNodeBFS(DeltaTime, Inputs, RootNodes, PhysicsProxy);
+			SimulateNodeBFS(DeltaTime, Inputs, RootNodes, PhysicsProxy, RootParticle);
 		}
 		else
 		{
 			for (int RootIndex : RootNodes)
 			{
-				SimulateNode(DeltaTime, Inputs, RootIndex, PhysicsProxy);
+				SimulateNode(DeltaTime, Inputs, RootIndex, PhysicsProxy, RootParticle);
 			}
 		}
 
@@ -252,7 +253,7 @@ void FSimModuleTree::Simulate(float DeltaTime, FAllInputs& Inputs, FClusterUnion
 
 }
 
-void FSimModuleTree::OnContactModification(FCollisionContactModifier& Modifier, FClusterUnionPhysicsProxy* PhysicsProxy)
+void FSimModuleTree::OnContactModification(FCollisionContactModifier& Modifier, IPhysicsProxyBase* PhysicsProxy)
 {
 	TArray<int> RootNodes;
 	GetRootNodes(RootNodes);
@@ -262,7 +263,7 @@ void FSimModuleTree::OnContactModification(FCollisionContactModifier& Modifier,
 	}
 }
 
-void FSimModuleTree::SimulateNode(float DeltaTime, FAllInputs& Inputs, int NodeIndex, FClusterUnionPhysicsProxy* PhysicsProxy)
+void FSimModuleTree::SimulateNode(float DeltaTime, FAllInputs& Inputs, int NodeIndex, IPhysicsProxyBase* PhysicsProxy, Chaos::FPBDRigidParticleHandle* ParticleHandle)
 {
 	if (ISimulationModuleBase* Module = AccessSimModule(NodeIndex))
 	{
@@ -270,11 +271,12 @@ void FSimModuleTree::SimulateNode(float DeltaTime, FAllInputs& Inputs, int NodeI
 		{
 			if (Module->IsEnabled())
 			{
-				Module->Simulate(PhysicsProxy, DeltaTime, Inputs, *this);
+				Module->Simulate(PhysicsProxy, ParticleHandle, DeltaTime, Inputs, *this);
 
 				if (IsAnimationEnabled() && Module->IsAnimationEnabled())
 				{
-					Module->Animate(PhysicsProxy);
+					Module->Animate();
+					UpdateClusterUnionTransformsIfRequired(PhysicsProxy, Module);
 				}
 			}
 		}
@@ -283,7 +285,7 @@ void FSimModuleTree::SimulateNode(float DeltaTime, FAllInputs& Inputs, int NodeI
 		{
 			for (int ChildIdx : GetChildren(NodeIndex))
 			{
-				SimulateNode(DeltaTime, Inputs, ChildIdx, PhysicsProxy);
+				SimulateNode(DeltaTime, Inputs, ChildIdx, PhysicsProxy, ParticleHandle);
 			}
 		}
 
@@ -291,11 +293,12 @@ void FSimModuleTree::SimulateNode(float DeltaTime, FAllInputs& Inputs, int NodeI
 		{
 			if (Module->IsEnabled())
 			{
-				Module->Simulate(PhysicsProxy, DeltaTime, Inputs, *this);
+				Module->Simulate(PhysicsProxy, ParticleHandle, DeltaTime, Inputs, *this);
 
 				if (IsAnimationEnabled() && Module->IsAnimationEnabled())
 				{
-					Module->Animate(PhysicsProxy);
+					Module->Animate();
+					UpdateClusterUnionTransformsIfRequired(PhysicsProxy, Module);
 				}
 			}
 		}
@@ -303,7 +306,35 @@ void FSimModuleTree::SimulateNode(float DeltaTime, FAllInputs& Inputs, int NodeI
 	}
 }
 
-void FSimModuleTree::SimulateNodeBFS(float DeltaTime, FAllInputs& Inputs, const TArray<int>& RootNodes, FClusterUnionPhysicsProxy* PhysicsProxy)
+
+void FSimModuleTree::UpdateClusterUnionTransformsIfRequired(IPhysicsProxyBase* PhysicsProxy, ISimulationModuleBase* Module)
+{
+	if (PhysicsProxy && Module && PhysicsProxy->GetType() == EPhysicsProxyType::ClusterUnionProxy)
+	{
+		if (FClusterUnionPhysicsProxy* CUProxy = static_cast<FClusterUnionPhysicsProxy*>(PhysicsProxy))
+		{
+			if (FPBDRigidClusteredParticleHandle* Particle = Module->GetClusterParticle(CUProxy))
+			{
+				const FSimModuleAnimationData& AnimData = Module->GetAnimationData();
+				uint16 AnimFlags = AnimData.AnimFlags;
+				if (AnimFlags == Chaos::EAnimationFlags::AnimateRotation)
+				{
+					const FQuat& RestRotation = Module->GetInitialParticleTransform().GetRotation();
+					Particle->ChildToParent().SetRotation(RestRotation * AnimData.CombinedRotation);
+				}
+
+				if (AnimFlags == Chaos::EAnimationFlags::AnimatePosition)
+				{
+					FVector Movement = Module->GetComponentTransform().TransformVector(AnimData.AnimationLocOffset);
+					FVector RestPos = Module->GetInitialParticleTransform().GetTranslation();
+					Particle->ChildToParent().SetTranslation(RestPos + Movement);
+				}
+			}
+		}
+	}
+}
+
+void FSimModuleTree::SimulateNodeBFS(float DeltaTime, FAllInputs& Inputs, const TArray<int>& RootNodes, IPhysicsProxyBase* PhysicsProxy, Chaos::FPBDRigidParticleHandle* ParticleHandle)
 {
 	TQueue<int> Queue;
 	TArray<int> Stack;
@@ -340,18 +371,19 @@ void FSimModuleTree::SimulateNodeBFS(float DeltaTime, FAllInputs& Inputs, const
 		{
 			if (Module->IsEnabled())
 			{
-				Module->Simulate(PhysicsProxy, DeltaTime, Inputs, *this);
+				Module->Simulate(PhysicsProxy, ParticleHandle, DeltaTime, Inputs, *this);
 
 				if (IsAnimationEnabled() && Module->IsAnimationEnabled())
 				{
-					Module->Animate(PhysicsProxy);
+					Module->Animate();
+					UpdateClusterUnionTransformsIfRequired(PhysicsProxy, Module);
 				}
 			}
 		}
 	}
 }
 
-void FSimModuleTree::OnContactModificationInternal(int NodeIndex, FCollisionContactModifier& Modifier, FClusterUnionPhysicsProxy* PhysicsProxy)
+void FSimModuleTree::OnContactModificationInternal(int NodeIndex, FCollisionContactModifier& Modifier, IPhysicsProxyBase* PhysicsProxy)
 {
 	if (ISimulationModuleBase* Module = AccessSimModule(NodeIndex))
 	{
@@ -398,88 +430,53 @@ void FSimModuleTree::GetRootNodes(TArray<int>& RootNodesOut)
 	}
 }
 
-void FSimModuleTree::UpdateModuleVelocites(FGeometryCollectionPhysicsProxy* PhysicsProxy)
+void FSimModuleTree::UpdateModuleVelocites(IPhysicsProxyBase* PhysicsProxy, Chaos::FPBDRigidParticleHandle* RootParticle, bool bWake)
 {
-	check(PhysicsProxy);
+	check(RootParticle);
+	Chaos::EnsureIsInPhysicsThreadContext();
 
-	// capture the velocities at the start of each sim iteration
-	for (int i = 0; i < SimulationModuleTree.Num(); i++)
+	if (PhysicsProxy && bWake && !SimulationModuleTree.IsEmpty())
 	{
-		if (ISimulationModuleBase* Module = SimulationModuleTree[i].SimModule)
+		if (PhysicsProxy->GetType() == EPhysicsProxyType::ClusterUnionProxy)
 		{
-			if (Chaos::FPBDRigidClusteredParticleHandle* ParentParticle = PhysicsProxy->GetSolverClusterHandle_Internal(Module->GetTransformIndex()))
+			if (FClusterUnionPhysicsProxy* CUProxy = static_cast<FClusterUnionPhysicsProxy*>(PhysicsProxy))
 			{
-				const FTransform BodyTransform(ParentParticle->GetR(), ParentParticle->GetX());
-
-				if (Module->IsBehaviourType(eSimModuleTypeFlags::Velocity))
+				if (bWake && !SimulationModuleTree.IsEmpty())
 				{
-					Chaos::FPBDRigidClusteredParticleHandle* Particle = nullptr;
-					if (Module->IsClustered())
-					{
-						Particle = PhysicsProxy->GetSolverClusterHandle_Internal(Module->GetTransformIndex());
-					}
-					else
-					{
-						Particle = PhysicsProxy->GetParticle_Internal(Module->GetTransformIndex());
-					}
+					const FPhysicsObjectHandle& PhysicsObject = CUProxy->GetPhysicsObjectHandle();
 
-					if (Particle)
-					{
-						FVector WorldLocation = BodyTransform.TransformPosition(Module->GetParentRelativeTransform().GetLocation());
-						const Chaos::FVec3 Arm = WorldLocation - Particle->GetX();
-
-						FVector WorldVelocity = Particle->GetV() - Chaos::FVec3::CrossProduct(Arm, Particle->GetW());
-						FVector LocalLinearVelocity = BodyTransform.InverseTransformVector(WorldVelocity);
-						LocalLinearVelocity = Module->GetClusteredTransform().InverseTransformVector(LocalLinearVelocity);
-
-						FVector LocalAngular = BodyTransform.InverseTransformVector(Particle->GetW());
-						LocalAngular = Module->GetClusteredTransform().InverseTransformVector(LocalAngular);
-
-						Module->SetLocalLinearVelocity(LocalLinearVelocity);
-						Module->SetLocalAngularVelocity(LocalAngular);
-					}
-			
+					Chaos::FWritePhysicsObjectInterface_Internal WriteInterface = Chaos::FPhysicsObjectInternalInterface::GetWrite();
+					WriteInterface.WakeUp({ &PhysicsObject, 1 });
 				}
 			}
+
 		}
-	}
 
-}
+		if (PhysicsProxy->GetType() == EPhysicsProxyType::SingleParticleProxy)
+		{
+			if (FSingleParticlePhysicsProxy* ParticleProxy = static_cast<FSingleParticlePhysicsProxy*>(PhysicsProxy))
+			{
+				const FPhysicsObjectHandle& PhysicsObject = ParticleProxy->GetPhysicsObject();
 
-void FSimModuleTree::UpdateModuleVelocites(FClusterUnionPhysicsProxy* PhysicsProxy, bool bWake)
-{
-	check(PhysicsProxy);
-	Chaos::EnsureIsInPhysicsThreadContext();
+				Chaos::FWritePhysicsObjectInterface_Internal WriteInterface = Chaos::FPhysicsObjectInternalInterface::GetWrite();
+				WriteInterface.WakeUp({ &PhysicsObject, 1 });
+			}
 
-	if (Chaos::FClusterUnionPhysicsProxy::FInternalParticle* ParentParticle = PhysicsProxy->GetParticle_Internal())
-	{
-		if (bWake && !SimulationModuleTree.IsEmpty())
-		{
-			const FPhysicsObjectHandle& PhysicsObject = PhysicsProxy->GetPhysicsObjectHandle();
-			
-			Chaos::FWritePhysicsObjectInterface_Internal WriteInterface = Chaos::FPhysicsObjectInternalInterface::GetWrite();
-			WriteInterface.WakeUp( { &PhysicsObject, 1 });
 		}
+	}
 
+	if (RootParticle)
+	{
 		// capture the velocities at the start of each sim iteration
 		for (int i = 0; i < SimulationModuleTree.Num(); i++)
 		{
 			if (ISimulationModuleBase* Module = SimulationModuleTree[i].SimModule)
 			{
-				const FTransform BodyTransform(ParentParticle->GetR(), ParentParticle->GetX());
+				const FTransform BodyTransform(RootParticle->GetR(), RootParticle->GetX());
 
 				if (Module->IsBehaviourType(eSimModuleTypeFlags::Velocity))
 				{
-					const Chaos::FClusterUnionPhysicsProxy::FInternalParticle* Particle = nullptr;
-					if (Module->IsClustered())
-					{
-						Particle = ParentParticle;
-					}
-					else
-					{
-						//				Particle = Particles[Module->GetTransformIndex()];
-					}
-
+					Chaos::FPBDRigidParticleHandle* Particle = RootParticle;
 					if (Particle)
 					{
 						const FTransform& OffsetTransform = Module->GetComponentTransform();
@@ -508,20 +505,20 @@ void FSimModuleTree::UpdateModuleVelocites(FClusterUnionPhysicsProxy* PhysicsPro
 	}
 }
 
-void FSimModuleTree::UpdateVehicleState(FClusterUnionPhysicsProxy* PhysicsProxy)
+void FSimModuleTree::UpdateVehicleState(Chaos::FPBDRigidParticleHandle* RootParticle)
 {
-	check(PhysicsProxy);
+	check(RootParticle);
 	Chaos::EnsureIsInPhysicsThreadContext();
 
-	if (Chaos::FClusterUnionPhysicsProxy::FInternalParticle* ParentParticle = PhysicsProxy->GetParticle_Internal())
+	if (RootParticle)
 	{
-		const FTransform BodyTransform(ParentParticle->GetR(), ParentParticle->GetX());
+		const FTransform BodyTransform(RootParticle->GetR(), RootParticle->GetX());
 
 		VehicleState.ForwardDir = BodyTransform.GetUnitAxis(EAxis::X);
 		VehicleState.UpDir = BodyTransform.GetUnitAxis(EAxis::Z);
 		VehicleState.RightDir = BodyTransform.GetUnitAxis(EAxis::Y);
-		VehicleState.ForwardSpeedKmh = Chaos::CmSToKmH(FVector::DotProduct(ParentParticle->GetV(), VehicleState.ForwardDir));
-		VehicleState.AngularVelocityRad = ParentParticle->GetW();
+		VehicleState.ForwardSpeedKmh = Chaos::CmSToKmH(FVector::DotProduct(RootParticle->GetV(), VehicleState.ForwardDir));
+		VehicleState.AngularVelocityRad = RootParticle->GetW();
 	}
 
 }
diff --git a/Engine/Source/Runtime/Experimental/ChaosVehicles/ChaosVehiclesCore/Private/SimModule/SimulationModuleBase.cpp b/Engine/Source/Runtime/Experimental/ChaosVehicles/ChaosVehiclesCore/Private/SimModule/SimulationModuleBase.cpp
--- a/Engine/Source/Runtime/Experimental/ChaosVehicles/ChaosVehiclesCore/Private/SimModule/SimulationModuleBase.cpp
+++ b/Engine/Source/Runtime/Experimental/ChaosVehicles/ChaosVehiclesCore/Private/SimModule/SimulationModuleBase.cpp
@@ -15,7 +15,7 @@ void ISimulationModuleBase::AddLocalForceAtPosition(const FVector& Force, const
 	AppliedForce = Force;
 	if (SimModuleTree)
 	{
-		SimModuleTree->AccessDeferredForces().Add(FDeferredForcesModular::FApplyForceAtPositionData(ComponentTransform, TransformIndex, ParticleIdx.Idx, Force, Position, bAllowSubstepping, bIsLocalForce, bLevelSlope, DebugColorIn));
+		SimModuleTree->AccessDeferredForces().Add(FDeferredForcesModular::FApplyForceAtPositionData(ComponentTransform, TransformIndex, ParticleIdx.Idx, Force, GetParentRelativeTransform().GetTranslation()+Position, bAllowSubstepping, bIsLocalForce, bLevelSlope, DebugColorIn));
 	}
 }
 
@@ -33,7 +33,7 @@ void ISimulationModuleBase::AddLocalForce(const FVector& Force, bool bAllowSubst
 	AppliedForce = Force;
 	if (SimModuleTree)
 	{
-		SimModuleTree->AccessDeferredForces().Add(FDeferredForcesModular::FApplyForceData(ComponentTransform, TransformIndex, ParticleIdx.Idx, Force, bAllowSubstepping, bIsLocalForce, bLevelSlope, DebugColorIn));
+		SimModuleTree->AccessDeferredForces().Add(FDeferredForcesModular::FApplyForceAtPositionData(ComponentTransform, TransformIndex, ParticleIdx.Idx, Force, GetParentRelativeTransform().GetTranslation(), bAllowSubstepping, bIsLocalForce, bLevelSlope, DebugColorIn));
 	}
 }
 
@@ -68,7 +68,7 @@ FVehicleBlackboard* ISimulationModuleBase::GetSimBlackboard()
 	return SimModuleTree ? SimModuleTree->GetSimBlackboard() : nullptr;
 }
 
-FPBDRigidClusteredParticleHandle* ISimulationModuleBase::GetClusterParticle(Chaos::FClusterUnionPhysicsProxy* Proxy)
+FPBDRigidClusteredParticleHandle* ISimulationModuleBase::GetClusterParticle(FClusterUnionPhysicsProxy* Proxy)
 { 
 	if (ParticleIdx.IsValid() && CachedParticle && (CachedParticle->UniqueIdx() == ParticleIdx))
 	{
@@ -113,9 +113,9 @@ FPBDRigidParticleHandle* ISimulationModuleBase::GetParticleFromUniqueIndex(int32
 
 void ISimulationModuleBase::SetAnimationData(const FName& BoneNameIn, const FVector& AnimationOffsetIn, int AnimationSetupIndexIn)
 {
-	BoneName = BoneNameIn;
+	AnimationData.BoneName = BoneNameIn;
+	AnimationData.AnimationSetupIndex = AnimationSetupIndexIn;
 	AnimationOffset = AnimationOffsetIn;
-	AnimationSetupIndex = AnimationSetupIndexIn;
 }
 
 bool ISimulationModuleBase::GetDebugString(FString& StringOut) const
@@ -143,7 +143,8 @@ void FSimOutputData::FillOutputState(const ISimulationModuleBase* SimModule)
 {
 	if (SimModule)
 	{
-		AnimationSetupIndex = SimModule->AnimationSetupIndex;
+		AnimationData = SimModule->GetAnimationData();
+		AnimationData.AnimationLocOffset += SimModule->GetAnimationOffset();
 	}
 
 #if !(UE_BUILD_SHIPPING || UE_BUILD_TEST)
@@ -157,16 +158,17 @@ void FSimOutputData::FillOutputState(const ISimulationModuleBase* SimModule)
 
 void FSimOutputData::Lerp(const FSimOutputData& InCurrent, const FSimOutputData& InNext, float Alpha)
 {
-	AnimationSetupIndex = InNext.AnimationSetupIndex;
-	AnimFlags = InNext.AnimFlags;
-	if (AnimFlags & EAnimationFlags::AnimatePosition)
+	AnimationData.AnimationSetupIndex = InNext.AnimationData.AnimationSetupIndex;
+	AnimationData.AnimFlags = InNext.AnimationData.AnimFlags;
+	AnimationData.CombinedRotation = InNext.AnimationData.CombinedRotation;
+	if (AnimationData.AnimFlags & EAnimationFlags::AnimatePosition)
 	{
-		AnimationLocOffset = FMath::Lerp(InCurrent.AnimationLocOffset, InNext.AnimationLocOffset, Alpha);
+		AnimationData.AnimationLocOffset = FMath::Lerp(InCurrent.AnimationData.AnimationLocOffset, InNext.AnimationData.AnimationLocOffset, Alpha);
 	}
 
-	if (AnimFlags & EAnimationFlags::AnimateRotation)
+	if (AnimationData.AnimFlags & EAnimationFlags::AnimateRotation)
 	{
-		AnimationRotOffset = FMath::Lerp(InCurrent.AnimationRotOffset, InNext.AnimationRotOffset, Alpha);
+		AnimationData.AnimationRotOffset = FMath::Lerp(InCurrent.AnimationData.AnimationRotOffset, InNext.AnimationData.AnimationRotOffset, Alpha);
 	}
 }
 
diff --git a/Engine/Source/Runtime/Experimental/ChaosVehicles/ChaosVehiclesCore/Private/SimModule/ThrusterModule.cpp b/Engine/Source/Runtime/Experimental/ChaosVehicles/ChaosVehiclesCore/Private/SimModule/ThrusterModule.cpp
--- a/Engine/Source/Runtime/Experimental/ChaosVehicles/ChaosVehiclesCore/Private/SimModule/ThrusterModule.cpp
+++ b/Engine/Source/Runtime/Experimental/ChaosVehicles/ChaosVehiclesCore/Private/SimModule/ThrusterModule.cpp
@@ -32,15 +32,10 @@ namespace Chaos
 		AddLocalForceAtPosition(Steer.RotateVector(Force), Setup().ForceOffset, true, false, false, FColor::Magenta);
 	}
 
-	void FThrusterSimModule::Animate(Chaos::FClusterUnionPhysicsProxy* Proxy)
+	void FThrusterSimModule::Animate()
 	{
-		if (FPBDRigidClusteredParticleHandle* ClusterChild = GetClusterParticle(Proxy))
-		{
-			FQuat Steer = FQuat(Setup().SteeringAxis, FMath::DegreesToRadians(SteerAngleDegrees));
-
-			ClusterChild->ChildToParent().SetRotation(GetInitialParticleTransform().GetRotation() * Steer);
-		}
-
+		AnimationData.AnimFlags = EAnimationFlags::AnimateRotation;
+		AnimationData.AnimationRotOffset.Yaw = SteerAngleDegrees;
 	}
 
 } // namespace Chaos
diff --git a/Engine/Source/Runtime/Experimental/ChaosVehicles/ChaosVehiclesCore/Private/SimModule/WheelModule.cpp b/Engine/Source/Runtime/Experimental/ChaosVehicles/ChaosVehiclesCore/Private/SimModule/WheelModule.cpp
--- a/Engine/Source/Runtime/Experimental/ChaosVehicles/ChaosVehiclesCore/Private/SimModule/WheelModule.cpp
+++ b/Engine/Source/Runtime/Experimental/ChaosVehicles/ChaosVehiclesCore/Private/SimModule/WheelModule.cpp
@@ -206,20 +206,20 @@ namespace Chaos
 		IntegrateAngularVelocity(DeltaTime, Setup().WheelInertia, Setup().MaxRotationVel);
 	}
 
-	void FWheelSimModule::Animate(Chaos::FClusterUnionPhysicsProxy* Proxy)
+	void FWheelSimModule::Animate()
 	{
-		if (Proxy)
-		{
-			if (FPBDRigidClusteredParticleHandle* ClusterChild = GetClusterParticle(Proxy))
-			{
-				float Direction = Setup().ReverseDirection ? -1.0f : 1.0f;
-				FQuat Rot = (Setup().Axis == Chaos::EWheelAxis::Y) ? FQuat(FVector(1, 0, 0), -GetAngularPosition() * Direction) : FQuat(FVector(0, 1, 0), GetAngularPosition() * Direction);
-				FQuat Steer = FQuat(FVector(0, 0, 1), FMath::DegreesToRadians(GetSteerAngleDegrees()));
+		float Direction = Setup().ReverseDirection ? -1.0f : 1.0f;
+		FQuat Rot = (Setup().Axis == Chaos::EWheelAxis::Y) ? FQuat(FVector(1, 0, 0), -GetAngularPosition() * Direction) : FQuat(FVector(0, 1, 0), GetAngularPosition() * Direction);
+		FQuat Steer = FQuat(FVector(0, 0, 1), FMath::DegreesToRadians(GetSteerAngleDegrees()));
+
+		float AngularPositionDegrees = -AngularPosition * 180.0f / PI;
+		float SteeringAngleDegrees = SteerAngleDegrees;
+
+		AnimationData.AnimFlags = EAnimationFlags::AnimateRotation;
+		AnimationData.AnimationRotOffset.Pitch = AngularPositionDegrees;
+		AnimationData.AnimationRotOffset.Yaw = SteeringAngleDegrees;
+		AnimationData.CombinedRotation = (Steer * Rot);
 
-				FTransform InitialTransform = GetInitialParticleTransform();
-				ClusterChild->ChildToParent().SetRotation(InitialTransform.GetRotation() * Steer * Rot);
-			}
-		}
 	}
 
 	bool FWheelSimModule::GetDebugString(FString& StringOut) const
@@ -245,10 +245,6 @@ namespace Chaos
 			RPM = Sim->GetRPM();
 			AngularPositionDegrees = -Sim->AngularPosition * 180.0f / PI;
 			SteeringAngleDegrees = Sim->SteerAngleDegrees;
-
-			AnimFlags = EAnimationFlags::AnimateRotation;
-			AnimationRotOffset.Pitch = AngularPositionDegrees;
-			AnimationRotOffset.Yaw = SteeringAngleDegrees;
 		}
 	}
 
diff --git a/Engine/Source/Runtime/Experimental/ChaosVehicles/ChaosVehiclesCore/Public/SimModule/AerofoilModule.h b/Engine/Source/Runtime/Experimental/ChaosVehicles/ChaosVehiclesCore/Public/SimModule/AerofoilModule.h
--- a/Engine/Source/Runtime/Experimental/ChaosVehicles/ChaosVehiclesCore/Public/SimModule/AerofoilModule.h
+++ b/Engine/Source/Runtime/Experimental/ChaosVehicles/ChaosVehiclesCore/Public/SimModule/AerofoilModule.h
@@ -88,7 +88,7 @@ namespace Chaos
 
 		virtual void Simulate(float DeltaTime, const FAllInputs& Inputs, FSimModuleTree& VehicleModuleSystem) override;
 
-		virtual void Animate(Chaos::FClusterUnionPhysicsProxy* Proxy) override;
+		virtual void Animate() override;
 
 		void SetDensityOfMedium(float InDensity)
 		{
diff --git a/Engine/Source/Runtime/Experimental/ChaosVehicles/ChaosVehiclesCore/Public/SimModule/DeferredForcesModular.h b/Engine/Source/Runtime/Experimental/ChaosVehicles/ChaosVehiclesCore/Public/SimModule/DeferredForcesModular.h
--- a/Engine/Source/Runtime/Experimental/ChaosVehicles/ChaosVehiclesCore/Public/SimModule/DeferredForcesModular.h
+++ b/Engine/Source/Runtime/Experimental/ChaosVehicles/ChaosVehiclesCore/Public/SimModule/DeferredForcesModular.h
@@ -214,6 +214,8 @@ public:
 	void Apply(TArray<Chaos::FPBDRigidParticleHandle*>& Particles
 		, TArray<Chaos::FPBDRigidClusteredParticleHandle*>& ClusterParticles);
 
+	void Apply(Chaos::FPBDRigidParticleHandle* Particle);
+
 private:
 
 	template<typename TransformType>
diff --git a/Engine/Source/Runtime/Experimental/ChaosVehicles/ChaosVehiclesCore/Public/SimModule/SimModuleTree.h b/Engine/Source/Runtime/Experimental/ChaosVehicles/ChaosVehiclesCore/Public/SimModule/SimModuleTree.h
--- a/Engine/Source/Runtime/Experimental/ChaosVehicles/ChaosVehiclesCore/Public/SimModule/SimModuleTree.h
+++ b/Engine/Source/Runtime/Experimental/ChaosVehicles/ChaosVehiclesCore/Public/SimModule/SimModuleTree.h
@@ -151,9 +151,9 @@ namespace Chaos
 
 		void AppendTreeUpdates(const FSimTreeUpdates& TreeUpdates);
 
-		void Simulate(float DeltaTime, FAllInputs& Inputs, FClusterUnionPhysicsProxy* PhysicsProxy);
+		void Simulate(float DeltaTime, FAllInputs& Inputs, IPhysicsProxyBase* PhysicsProxy, Chaos::FPBDRigidParticleHandle* RootParticle);
 
-		void OnContactModification(FCollisionContactModifier& Modifier, FClusterUnionPhysicsProxy* PhysicsProxy);
+		void OnContactModification(FCollisionContactModifier& Modifier, IPhysicsProxyBase* PhysicsProxy);
 
 		void SetSimTreeProcessingOrder(ESimTreeProcessingOrder OrderIn) { SimTreeProcessingOrder = OrderIn; }
 
@@ -213,19 +213,21 @@ namespace Chaos
 		}
 
 	protected:
-		void SimulateNode(float DeltaTime, FAllInputs& Inputs, int NodeIdx, FClusterUnionPhysicsProxy* PhysicsProxy);
+		void SimulateNode(float DeltaTime, FAllInputs& Inputs, int NodeIdx, IPhysicsProxyBase* PhysicsProxy, Chaos::FPBDRigidParticleHandle* ParticleHandle);
 
-		void OnContactModificationInternal(int NodeIndex, FCollisionContactModifier& Modifier, FClusterUnionPhysicsProxy* PhysicsProxy);
+		void OnContactModificationInternal(int NodeIndex, FCollisionContactModifier& Modifier, IPhysicsProxyBase* PhysicsProxy);
 
-		void SimulateNodeBFS(float DeltaTime, FAllInputs& Inputs, const TArray<int>& RootNodes, FClusterUnionPhysicsProxy* PhysicsProxy);
+		void SimulateNodeBFS(float DeltaTime, FAllInputs& Inputs, const TArray<int>& RootNodes, IPhysicsProxyBase* PhysicsProxy, Chaos::FPBDRigidParticleHandle* ParticleHandle);
 
 		void DeleteNodesBelow(int NodeIdx);
 
 		int GetNextIndex();
 
-		void UpdateModuleVelocites(FGeometryCollectionPhysicsProxy* PhysicsProxy);
-		void UpdateModuleVelocites(FClusterUnionPhysicsProxy* PhysicsProxy, bool bWake);
-		void UpdateVehicleState(FClusterUnionPhysicsProxy* PhysicsProxy);
+		void UpdateClusterUnionTransformsIfRequired(IPhysicsProxyBase* PhysicsProxy, ISimulationModuleBase* Module);
+
+
+		void UpdateModuleVelocites(IPhysicsProxyBase* PhysicsProxy, Chaos::FPBDRigidParticleHandle* RootParticle, bool bWake);
+		void UpdateVehicleState(Chaos::FPBDRigidParticleHandle* RootParticle);
 
 		TArray<FSimModuleNode> SimulationModuleTree;
 		TArray<int> FreeList;
diff --git a/Engine/Source/Runtime/Experimental/ChaosVehicles/ChaosVehiclesCore/Public/SimModule/SimulationModuleBase.h b/Engine/Source/Runtime/Experimental/ChaosVehicles/ChaosVehiclesCore/Public/SimModule/SimulationModuleBase.h
--- a/Engine/Source/Runtime/Experimental/ChaosVehicles/ChaosVehiclesCore/Public/SimModule/SimulationModuleBase.h
+++ b/Engine/Source/Runtime/Experimental/ChaosVehicles/ChaosVehiclesCore/Public/SimModule/SimulationModuleBase.h
@@ -4,6 +4,7 @@
 
 #include "Chaos/GeometryParticlesfwd.h"
 #include "Chaos/ParticleHandleFwd.h"
+#include "Chaos/PhysicsObject.h"
 #include "Logging/LogMacros.h"
 #include "SimModule/ModuleFactoryRegister.h"
 #include "SimModule/ModuleInput.h"
@@ -43,6 +44,13 @@ namespace Chaos
 	const FName RollControlName("Roll");
 	const FName YawControlName("Yaw");
 
+	namespace EAnimationFlags
+	{
+		static uint16 AnimateNone = 0x00000000;
+		static uint16 AnimatePosition = 0x00000001;
+		static uint16 AnimateRotation = 0x00000002;
+	}
+
 	struct CHAOSVEHICLESCORE_API FModuleHitResults
 	{
 		int SimIndex;
@@ -61,6 +69,26 @@ namespace Chaos
 		bool bKeepVehicleAwake = false;
 	};
 
+	struct CHAOSVEHICLESCORE_API FSimModuleAnimationData
+	{
+		FSimModuleAnimationData() 
+			: BoneName(NAME_None)
+			, AnimFlags(EAnimationFlags::AnimateNone)
+			, AnimationLocOffset(FVector::ZeroVector)
+			, AnimationRotOffset(FRotator::ZeroRotator)
+			, CombinedRotation(FQuat::Identity)
+			, AnimationSetupIndex(-1)
+			{}
+
+		FName BoneName;
+		uint16 AnimFlags = EAnimationFlags::AnimateNone;
+		FVector AnimationLocOffset;
+		FRotator AnimationRotOffset;
+		FQuat CombinedRotation;
+		int AnimationSetupIndex;
+	};
+
+
 	/**
 	 * Code common between all simulation building blocks settings
 	 */
@@ -250,8 +278,6 @@ namespace Chaos
 
 		ISimulationModuleBase()
 			: SimModuleTree(nullptr)
-			, BoneName(NAME_None)
-			, AnimationSetupIndex(-1)
 			, SimTreeIndex(INVALID_IDX)
 			, StateFlags(Enabled)
 			, TransformIndex(INVALID_IDX)
@@ -291,7 +317,7 @@ namespace Chaos
 		/*
 		* Any post construction initialisation - called from Game Thread
 		*/
-		virtual void OnConstruction_External(Chaos::FClusterUnionPhysicsProxy* Proxy) {}
+		virtual void OnConstruction_External(const Chaos::FPhysicsObjectHandle& PhysicsObject) {}
 
 		/*
 		* Any cleaning up required - called from game thread
@@ -301,12 +327,12 @@ namespace Chaos
 		/**
 		 * The main Simulation function that is called from the physics async callback thread
 		 */
-		virtual void Simulate(Chaos::FClusterUnionPhysicsProxy* Proxy, float DeltaTime, const FAllInputs& Inputs, FSimModuleTree& VehicleModuleSystem) 
+		virtual void Simulate(IPhysicsProxyBase* Proxy, Chaos::FPBDRigidParticleHandle* ParticleHandle, float DeltaTime, const FAllInputs& Inputs, FSimModuleTree& VehicleModuleSystem)
 		{
 			Simulate(DeltaTime, Inputs, VehicleModuleSystem);
 		}
 
-		virtual void OnContactModification(Chaos::FCollisionContactModifier& Modifier, Chaos::FClusterUnionPhysicsProxy* Proxy) {}
+		virtual void OnContactModification(Chaos::FCollisionContactModifier& Modifier, IPhysicsProxyBase* Proxy) {}
 
 		/**
 		* The main Simulation function that is called from the physics async callback thread
@@ -316,15 +342,16 @@ namespace Chaos
 		/**
 		 * Animate/modify the childToParent transforms, to say rotate a wheel, or rudder, etc
 		 */
-		virtual void Animate(Chaos::FClusterUnionPhysicsProxy* Proxy) {}
+		virtual void Animate() {}
 
 		void SetAnimationEnabled(bool bInEnabled) { bAnimationEnabled = bInEnabled; }
 		bool IsAnimationEnabled() { return bAnimationEnabled; }
 
 		void SetAnimationData(const FName& BoneNameIn, const FVector& AnimationOffsetIn, int AnimationSetupIndexIn);
+		const FSimModuleAnimationData& GetAnimationData() const { return AnimationData; }
 		const FVector& GetAnimationOffset() const { return AnimationOffset; }
-		const FName& GetBoneName() const { return BoneName; }
-		int GetAnimationSetupIndex() { return AnimationSetupIndex; }
+		const FName& GetBoneName() const { return AnimationData.BoneName; }
+		int GetAnimationSetupIndex() { return AnimationData.AnimationSetupIndex; }
 
 		/**
 		 * Option to draw debug for this module requires CVar p.Chaos.DebugDraw.Enabled 1
@@ -438,12 +465,9 @@ namespace Chaos
 
 		Chaos::FPBDRigidParticleHandle* GetParticleFromUniqueIndex(int32 ParticleUniqueIdx, TArray<Chaos::FPBDRigidParticleHandle*>& Particles);
 
-
 	protected:
 
 		FSimModuleTree* SimModuleTree;	// A pointer back to the simulation tree where we are stored
-		FName BoneName;
-		int AnimationSetupIndex;
 		int SimTreeIndex;	// Index of this SimModule in the FSimModuleTree
 		eSimModuleState StateFlags;	// TODO: make this more like flags
 		int TransformIndex; // Index of this Sim Module's node in Geometry Collection Transform array
@@ -460,7 +484,7 @@ namespace Chaos
 		bool bClustered;
 		bool bAnimationEnabled;
 		FVector AnimationOffset; 
-		FVector AnimationRotation;
+		FSimModuleAnimationData AnimationData;
 
 		// for headless chaos testing
 		FVector AppliedForce;
@@ -508,13 +532,6 @@ namespace Chaos
 
 	using FModuleNetDataArray = TArray<TSharedPtr<FModuleNetData>>;
 
-	namespace EAnimationFlags
-	{
-		static uint16 AnimateNone = 0x00000000;
-		static uint16 AnimatePosition = 0x00000001;
-		static uint16 AnimateRotation = 0x00000002;
-	}
-	
 	struct CHAOSVEHICLESCORE_API FSimOutputData : public FSimulationModuleTypeContainer, public TSimulationModuleTypeableBase<FSimOutputData>
 	{
 		DEFINE_CHAOSSIMTYPENAME(FSimOutputData);
@@ -527,10 +544,7 @@ namespace Chaos
 		virtual void Lerp(const FSimOutputData& InCurrent, const FSimOutputData& InNext, float Alpha);
 
 		bool bEnabled = true;
-		int AnimationSetupIndex;
-		uint16 AnimFlags;
-		FVector AnimationLocOffset;
-		FRotator AnimationRotOffset;
+		FSimModuleAnimationData AnimationData;
 
 #if !(UE_BUILD_SHIPPING || UE_BUILD_TEST)
 		virtual FString ToString() { return FString(); }
diff --git a/Engine/Source/Runtime/Experimental/ChaosVehicles/ChaosVehiclesCore/Public/SimModule/ThrusterModule.h b/Engine/Source/Runtime/Experimental/ChaosVehicles/ChaosVehiclesCore/Public/SimModule/ThrusterModule.h
--- a/Engine/Source/Runtime/Experimental/ChaosVehicles/ChaosVehiclesCore/Public/SimModule/ThrusterModule.h
+++ b/Engine/Source/Runtime/Experimental/ChaosVehicles/ChaosVehiclesCore/Public/SimModule/ThrusterModule.h
@@ -58,7 +58,7 @@ namespace Chaos
 
 		virtual void Simulate(float DeltaTime, const FAllInputs& Inputs, FSimModuleTree& VehicleModuleSystem) override;
 
-		virtual void Animate(Chaos::FClusterUnionPhysicsProxy* Proxy) override;
+		virtual void Animate() override;
 
 		float GetSteerAngleDegrees() const { return SteerAngleDegrees; }
 
diff --git a/Engine/Source/Runtime/Experimental/ChaosVehicles/ChaosVehiclesCore/Public/SimModule/WheelModule.h b/Engine/Source/Runtime/Experimental/ChaosVehicles/ChaosVehiclesCore/Public/SimModule/WheelModule.h
--- a/Engine/Source/Runtime/Experimental/ChaosVehicles/ChaosVehiclesCore/Public/SimModule/WheelModule.h
+++ b/Engine/Source/Runtime/Experimental/ChaosVehicles/ChaosVehiclesCore/Public/SimModule/WheelModule.h
@@ -146,7 +146,7 @@ namespace Chaos
 
 		virtual void Simulate(float DeltaTime, const FAllInputs& Inputs, FSimModuleTree& VehicleModuleSystem) override;
 
-		virtual void Animate(Chaos::FClusterUnionPhysicsProxy* Proxy) override;
+		virtual void Animate() override;
 
 		virtual bool IsBehaviourType(eSimModuleTypeFlags InType) const override { return (InType & TorqueBased) || (InType & Velocity); }
 
