Simple bug fix, ensure looks to be invalid, also testing against something that isn't used in function

diff --git a/Engine/Plugins/Experimental/ChaosModularVehicle/Source/ChaosModularVehicleEngine/Private/ChaosModularVehicle/ChaosSimModuleManager.cpp b/Engine/Plugins/Experimental/ChaosModularVehicle/Source/ChaosModularVehicleEngine/Private/ChaosModularVehicle/ChaosSimModuleManager.cpp
--- a/Engine/Plugins/Experimental/ChaosModularVehicle/Source/ChaosModularVehicleEngine/Private/ChaosModularVehicle/ChaosSimModuleManager.cpp
+++ b/Engine/Plugins/Experimental/ChaosModularVehicle/Source/ChaosModularVehicleEngine/Private/ChaosModularVehicle/ChaosSimModuleManager.cpp
@@ -204,7 +204,6 @@ void FChaosSimModuleManager::InjectInputs_External(int32 PhysicsStep, int32 NumS
 	}
 	FChaosSimModuleManagerAsyncInput* AsyncInput = AsyncCallback->GetProducerInputData_External();
 	check(AsyncInput);
-	ensure(AsyncInput->World == World);
 
 	for (TWeakObjectPtr<UModularVehicleBaseComponent> Vehicle : CUVehicles)
 	{
