Bug fix for assert in UPawnMovementComponent::SetUpdatedComponent

diff --git a/Engine/Plugins/Experimental/ChaosModularVehicle/Source/ChaosModularVehicleEngine/Private/ChaosModularVehicle/ModularVehicleClusterActor.cpp b/Engine/Plugins/Experimental/ChaosModularVehicle/Source/ChaosModularVehicleEngine/Private/ChaosModularVehicle/ModularVehicleClusterActor.cpp
--- a/Engine/Plugins/Experimental/ChaosModularVehicle/Source/ChaosModularVehicleEngine/Private/ChaosModularVehicle/ModularVehicleClusterActor.cpp
+++ b/Engine/Plugins/Experimental/ChaosModularVehicle/Source/ChaosModularVehicleEngine/Private/ChaosModularVehicle/ModularVehicleClusterActor.cpp
@@ -14,7 +14,8 @@ AModularVehicleClusterActor::AModularVehicleClusterActor(const FObjectInitialize
 	// Vehicle sim piggy-backs off of the cluster union component & its events
 	VehicleSimComponent = CreateDefaultSubobject<UModularVehicleBaseComponent>(TEXT("VehicleSimComponent0"));
 	VehicleSimComponent->SetClusterComponent(ClusterUnionVehicleComponent);
-
+	//VehicleSimComponent->bAutoRegisterUpdatedComponent = false;
+                          
 	SetRemoteRoleForBackwardsCompat(ROLE_SimulatedProxy);
 	bReplicates = true;
 	SetReplicatingMovement(true);
