Modular Vehicle - change underlying raycast function to hopefully avoid frame delay

diff --git a/Engine/Plugins/Experimental/ChaosModularVehicle/Source/ChaosModularVehicleEngine/Private/ChaosModularVehicle/ModularVehicleSimulationCU.cpp b/Engine/Plugins/Experimental/ChaosModularVehicle/Source/ChaosModularVehicleEngine/Private/ChaosModularVehicle/ModularVehicleSimulationCU.cpp
--- a/Engine/Plugins/Experimental/ChaosModularVehicle/Source/ChaosModularVehicleEngine/Private/ChaosModularVehicle/ModularVehicleSimulationCU.cpp
+++ b/Engine/Plugins/Experimental/ChaosModularVehicle/Source/ChaosModularVehicleEngine/Private/ChaosModularVehicle/ModularVehicleSimulationCU.cpp
@@ -197,6 +197,7 @@ void FModularVehicleSimulation::PerformAdditionalSimWork(UWorld* InWorld, const
 					FVector TraceVector(TraceStart - TraceEnd);
 					FVector TraceNormal = TraceVector.GetSafeNormal();
 
+					TArray<FHitResult> HitResultsOut;
 					FHitResult HitResult = FHitResult();
 					ECollisionChannel SpringCollisionChannel = InputData.PhysicsInputs.CollisionChannel;
 					const FCollisionResponseParams& ResponseParams = InputData.PhysicsInputs.TraceCollisionResponse;
@@ -206,19 +207,29 @@ void FModularVehicleSimulation::PerformAdditionalSimWork(UWorld* InWorld, const
 						{
 							case ETraceType::Spherecast:
 							{
-								InWorld->SweepSingleByChannel(HitResult
+								float QueryRadius = WheelRadius;
+								if (Chaos::Private::FGenericPhysicsInterface_Internal::SpherecastMulti(InWorld
+									, QueryRadius
+									, HitResultsOut
 									, TraceStart + TraceNormal * WheelRadius
 									, TraceEnd + TraceNormal * WheelRadius
-									, FQuat::Identity, SpringCollisionChannel
-									, FCollisionShape::MakeSphere(WheelRadius), TraceParams
-									, ResponseParams);
+									, SpringCollisionChannel
+									, TraceParams
+									, ResponseParams))
+								{
+									HitResult = HitResultsOut.Last();
+								}
 							}
 							break;
 
 							case ETraceType::Raycast:
 							default:
 							{
-								InWorld->LineTraceSingleByChannel(HitResult, TraceStart, TraceEnd, SpringCollisionChannel, TraceParams, ResponseParams);
+								float QueryRadius = 0.0f;
+								if (Chaos::Private::FGenericPhysicsInterface_Internal::SpherecastMulti(InWorld, QueryRadius, HitResultsOut, TraceStart, TraceEnd, SpringCollisionChannel, TraceParams, ResponseParams))
+								{
+									HitResult = HitResultsOut.Last();
+								}
 							}
 							break;
 						}
