Changes to make Modular Vehicle and Cluster Union to work with Static/Skeletal meshes (SingleParticlePhysicsProxy) as well as GeometryCollection.

diff --git a/Engine/Plugins/Experimental/ChaosModularVehicle/Source/ChaosModularVehicleEngine/Private/ChaosModularVehicle/ModularVehicleBaseComponent.cpp b/Engine/Plugins/Experimental/ChaosModularVehicle/Source/ChaosModularVehicleEngine/Private/ChaosModularVehicle/ModularVehicleBaseComponent.cpp
--- a/Engine/Plugins/Experimental/ChaosModularVehicle/Source/ChaosModularVehicleEngine/Private/ChaosModularVehicle/ModularVehicleBaseComponent.cpp
+++ b/Engine/Plugins/Experimental/ChaosModularVehicle/Source/ChaosModularVehicleEngine/Private/ChaosModularVehicle/ModularVehicleBaseComponent.cpp
@@ -12,6 +12,7 @@
 #include "PhysicsEngine/PhysicsSettings.h"
 #include "PhysicsEngine/ClusterUnionComponent.h"
 #include "PhysicsEngine/PhysicsObjectExternalInterface.h"
+#include "PhysicsEngine/BodySetup.h"
 #include "Net/Core/PushModel/PushModel.h"
 #include "Net/UnrealNetwork.h"
 #include "GameFramework/PlayerController.h"
@@ -21,6 +22,8 @@
 #include "PhysicsProxy/SingleParticlePhysicsProxy.h"
 #include "PhysicsProxy/ClusterUnionPhysicsProxy.h"
 #include "GeometryCollection/GeometryCollectionComponent.h"
+#include "Components/StaticMeshComponent.h"
+#include "Components/SkeletalMeshComponent.h"
 
 #include "ChaosModularVehicle/ChaosSimModuleManagerAsyncCallback.h"
 
@@ -268,7 +271,7 @@ void UModularVehicleBaseComponent::CacheRootPhysicsObject(IPhysicsProxyBase* Pro
 	}
 }
 
-void UModularVehicleBaseComponent::CreateAssociatedSimComponents(USceneComponent* AttachedComponent, int ParentIndex, int TransformIndex, Chaos::FSimTreeUpdates& TreeUpdatesOut)
+void UModularVehicleBaseComponent::CreateAssociatedSimComponents(USceneComponent* ParentComponent, USceneComponent* AttachedComponent, int ParentIndex, int InTransformIndex, Chaos::FSimTreeUpdates& TreeUpdatesOut)
 {
 	using namespace Chaos;
 	if (AttachedComponent == nullptr || ClusterUnionComponent == nullptr)
@@ -281,14 +284,58 @@ void UModularVehicleBaseComponent::CreateAssociatedSimComponents(USceneComponent
 		CacheRootPhysicsObject(GetPhysicsProxy());
 	}
 
-	UE_LOG(LogModularBase, Log, TEXT("CreateAssociatedSimComponents Attaching %s, TransformIndex %d"), *AttachedComponent->GetName(), TransformIndex);
+	UE_LOG(LogModularBase, Log, TEXT("CreateAssociatedSimComponents Attaching %s, TransformIndex %d"), *AttachedComponent->GetName(), InTransformIndex);
 
 	TArray<Chaos::FClusterUnionChildData> ChildParticles = ClusterUnionComponent->GetPhysicsProxy()->GetSyncedData_External().ChildParticles;
 
+	int TransformIndex = InTransformIndex;
+
 	ensure(TransformIndex < ChildParticles.Num());
 
 	if (IVehicleSimBaseComponentInterface* ComponentInterface = Cast<IVehicleSimBaseComponentInterface>(AttachedComponent))
 	{
+
+		// if we find the parent is a skeletal mesh then the indices are ...
+		if (ParentComponent)
+		{
+			if (USkeletalMeshComponent* SKMComponent = Cast<USkeletalMeshComponent>(ParentComponent))
+			{
+				const FName BoneName = ComponentInterface->GetBoneName();
+
+				int32 BodyBoneIndex = INDEX_NONE;
+
+				for (int32 ArrayIndex = 0; ArrayIndex < SKMComponent->Bodies.Num(); ArrayIndex++)
+				{
+					// Get body index of bone
+					if (SKMComponent->Bodies[ArrayIndex] && SKMComponent->Bodies[ArrayIndex]->IsValidBodyInstance())
+					{
+						const FName BodyName = SKMComponent->Bodies[ArrayIndex]->GetBodySetup()->BoneName;
+
+						if (BodyName == BoneName)
+						{
+							BodyBoneIndex = ArrayIndex;
+							break;
+						}
+
+					}
+				}
+
+				if (BodyBoneIndex != INDEX_NONE)
+				{
+					TransformIndex = InTransformIndex + BodyBoneIndex;
+				}
+
+			}
+			else
+			{
+				TransformIndex = InTransformIndex;
+			}
+
+		}
+
+		ensure(TransformIndex < ChildParticles.Num());
+
+
 		FTransform ClusterUnionComponentTransform = ClusterUnionComponent->GetComponentTransform();
 		FTransform ComponentTransform = AttachedComponent->GetComponentTransform().GetRelativeTransform(ClusterUnionComponentTransform);
 
@@ -337,7 +384,7 @@ void UModularVehicleBaseComponent::CreateAssociatedSimComponents(USceneComponent
 		{
 			// if bone already exists then use that (seperate wheel and suspension modules can share same bone)
 
-			int FoundIndex = -1;
+			int FoundIndex = INDEX_NONE;
 			if (ComponentBoneName != NAME_None)
 			{
 				for (int I = 0; I < ModuleAnimationSetups.Num(); I++)
@@ -351,7 +398,7 @@ void UModularVehicleBaseComponent::CreateAssociatedSimComponents(USceneComponent
 			}
 
 			const FVector& ComponentAnimationOffset = ComponentInterface->GetAnimationOffset();
-			if (FoundIndex != -1)
+			if (FoundIndex != INDEX_NONE)
 			{
 				NewModule->SetAnimationData(ComponentBoneName, ComponentAnimationOffset, FoundIndex);
 			}
@@ -382,7 +429,7 @@ void UModularVehicleBaseComponent::CreateAssociatedSimComponents(USceneComponent
 	{
 		if (IVehicleSimBaseComponentInterface* ChildSimComponent = Cast<IVehicleSimBaseComponentInterface>(Child))
 		{
-			CreateAssociatedSimComponents(Child, ParentIndex, TransformIndex, TreeUpdatesOut);
+			CreateAssociatedSimComponents(AttachedComponent, Child, ParentIndex, TransformIndex, TreeUpdatesOut);
 		}
 	}
 
@@ -454,7 +501,7 @@ void UModularVehicleBaseComponent::CreateIndependentSimComponents(USceneComponen
 		{
 			// if bone already exists then use that (seperate wheel and suspension modules can share same bone)
 
-			int FoundIndex = -1;
+			int FoundIndex = INDEX_NONE;
 			for (int I = 0; I < ModuleAnimationSetups.Num(); I++)
 			{
 				if (ModuleAnimationSetups[I].BoneName == ComponentBoneName)
@@ -465,7 +512,7 @@ void UModularVehicleBaseComponent::CreateIndependentSimComponents(USceneComponen
 			}
 
 			const FVector& ComponentAnimationOffset = ComponentInterface->GetAnimationOffset();
-			if (FoundIndex != -1)
+			if (FoundIndex != INDEX_NONE)
 			{
 				NewModule->SetAnimationData(ComponentBoneName, ComponentAnimationOffset, FoundIndex);
 			}
@@ -597,7 +644,7 @@ void UModularVehicleBaseComponent::BeginPlay()
 		}
 		else
 		{
-			AddGeometryCollectionsFromOwnedActor();
+			AddComponentsFromOwnedActor();
 		}
 	}
 
@@ -1113,7 +1160,7 @@ void UModularVehicleBaseComponent::AddComponentToSimulationImpl(UPrimitiveCompon
 		int ParentID = FindParentsLastSimComponent(InComponent);
 
 		Chaos::FSimTreeUpdates LatestTreeUpdates;
-		CreateAssociatedSimComponents(InComponent, ParentID, NextTransformIndex, LatestTreeUpdates);
+		CreateAssociatedSimComponents(nullptr, InComponent, ParentID, NextTransformIndex, LatestTreeUpdates);
 
 		Chaos::FSimTreeUpdates* NextTreeUpdates = nullptr;
 		if (LastComponentAddIndex + 1 == ComponentAddOrder)
@@ -1129,7 +1176,13 @@ void UModularVehicleBaseComponent::AddComponentToSimulationImpl(UPrimitiveCompon
 			PendingTreeUpdates.Add(InComponent, MoveTemp(LatestTreeUpdates));
 		}
 
-		NextTransformIndex++;
+		int Size = 1; // for UStaticMeshComponent && UGeometryCollectionComponent cases
+		if (USkeletalMeshComponent* SKMComponent = Cast<USkeletalMeshComponent>(InComponent))
+		{
+			Size = SKMComponent->Bodies.Num();
+		}
+
+		NextTransformIndex += Size;
 	}
 }
 
@@ -1311,6 +1364,45 @@ void UModularVehicleBaseComponent::DestroyVehicleSim()
 
 }
 
+
+void UModularVehicleBaseComponent::AddComponentToCluster(USceneComponent* Component)
+{
+	TArray<int32> BoneIds;
+
+	if (UGeometryCollectionComponent* GCComponent = Cast<UGeometryCollectionComponent>(Component))
+	{
+		BoneIds.Add(0);
+
+		UE_LOG(LogModularBase, Log, TEXT("AddComponentToCluster: Geometry Collection"));
+		ClusterUnionComponent->AddComponentToCluster(GCComponent, BoneIds);
+	}
+
+	if (UStaticMeshComponent* SMComponent = Cast<UStaticMeshComponent>(Component))
+	{
+		if (SMComponent->HasValidPhysicsState())
+		{
+			BoneIds.Add(0);
+
+			UE_LOG(LogModularBase, Log, TEXT("AddComponentToCluster: Static Mesh"));
+			ClusterUnionComponent->AddComponentToCluster(SMComponent, BoneIds);
+		}
+	}
+
+	if (USkeletalMeshComponent* SKMComponent = Cast<USkeletalMeshComponent>(Component))
+	{
+		for (int I = 0; I < SKMComponent->Bodies.Num(); I++)
+		{
+			BoneIds.Add(I);
+		}
+		if (BoneIds.Num() > 0)
+		{
+			UE_LOG(LogModularBase, Log, TEXT("AddComponentToCluster: Skeletal Mesh"));
+			ClusterUnionComponent->AddComponentToCluster(SKMComponent, BoneIds);
+		}
+	}
+}
+
+
 void UModularVehicleBaseComponent::AddOverlappingComponentsToCluster()
 {
 	//if (UWorld* World = GetWorld())
@@ -1339,7 +1431,7 @@ void UModularVehicleBaseComponent::AddOverlappingComponentsToCluster()
 	//}
 }
 
-bool UModularVehicleBaseComponent::AddGeometryCollectionsFromOwnedActor()
+bool UModularVehicleBaseComponent::AddComponentsFromOwnedActor()
 {
 	bool bSuccess = false;
 
@@ -1355,11 +1447,7 @@ bool UModularVehicleBaseComponent::AddGeometryCollectionsFromOwnedActor()
 
 			for (USceneComponent* Component : ChildComponents)
 			{
-				if (UGeometryCollectionComponent* GCComponent = Cast<UGeometryCollectionComponent>(Component))
-				{
-					ClusterUnionComponent->AddComponentToCluster(GCComponent, BoneIds);
-					bSuccess = true;
-				}
+				AddComponentToCluster(Component);
 			}
 		}
 	}
diff --git a/Engine/Plugins/Experimental/ChaosModularVehicle/Source/ChaosModularVehicleEngine/Public/ChaosModularVehicle/ModularVehicleBaseComponent.h b/Engine/Plugins/Experimental/ChaosModularVehicle/Source/ChaosModularVehicleEngine/Public/ChaosModularVehicle/ModularVehicleBaseComponent.h
--- a/Engine/Plugins/Experimental/ChaosModularVehicle/Source/ChaosModularVehicleEngine/Public/ChaosModularVehicle/ModularVehicleBaseComponent.h
+++ b/Engine/Plugins/Experimental/ChaosModularVehicle/Source/ChaosModularVehicleEngine/Public/ChaosModularVehicle/ModularVehicleBaseComponent.h
@@ -132,7 +132,7 @@ public:
 	void ProduceInput(int32 PhysicsStep, int32 NumSteps);
 
 	// Create sim components hierarchy from parent component - used by original cluster code union path
-	virtual void CreateAssociatedSimComponents(USceneComponent* AttachedComponent, int ParentIndex, int TransformIndex, Chaos::FSimTreeUpdates& TreeUpdatesOut);
+	virtual void CreateAssociatedSimComponents(USceneComponent* ParentComponent, USceneComponent* AttachedComponent, int ParentIndex, int TransformIndex, Chaos::FSimTreeUpdates& TreeUpdatesOut);
 
 	// Create sim components in hierarchy from parent component - used by code path that does not rely on cluster union
 	virtual void CreateIndependentSimComponents(USceneComponent* RootComponent, USceneComponent* AttachedComponent, int ParentIndex, int TransformIndex, Chaos::FSimTreeUpdates& TreeUpdatesOut);
@@ -291,7 +291,8 @@ protected:
 	void DestroyVehicleSim();
 	void UpdatePhysicalProperties();
 	void AddOverlappingComponentsToCluster();
-	bool AddGeometryCollectionsFromOwnedActor();
+	void AddComponentToCluster(USceneComponent* Component);
+	bool AddComponentsFromOwnedActor();
 	void SetupSkeletalAnimationStructure();
 	void AssimilateComponentInputs(TArray<FModuleInputSetup>& OutCombinedInputs);
 	void CacheRootPhysicsObject(IPhysicsProxyBase* Proxy);
diff --git a/Engine/Source/Runtime/Engine/Classes/Components/SkeletalMeshComponent.h b/Engine/Source/Runtime/Engine/Classes/Components/SkeletalMeshComponent.h
--- a/Engine/Source/Runtime/Engine/Classes/Components/SkeletalMeshComponent.h
+++ b/Engine/Source/Runtime/Engine/Classes/Components/SkeletalMeshComponent.h
@@ -1771,6 +1771,8 @@ public:
 	//~ End UObject Interface.
 
 	//~ Begin UActorComponent Interface.
+	ENGINE_API virtual bool HasValidPhysicsState() const override;
+
 protected:
 	ENGINE_API virtual void OnRegister() override;
 	ENGINE_API virtual void OnUnregister() override;
diff --git a/Engine/Source/Runtime/Engine/Private/PhysicsEngine/BodyInstance.cpp b/Engine/Source/Runtime/Engine/Private/PhysicsEngine/BodyInstance.cpp
--- a/Engine/Source/Runtime/Engine/Private/PhysicsEngine/BodyInstance.cpp
+++ b/Engine/Source/Runtime/Engine/Private/PhysicsEngine/BodyInstance.cpp
@@ -1290,11 +1290,11 @@ void TInitBodiesHelperBase<InAllocatorType>::CreateActor_AssumesLocked(FBodyInst
 	{
 		ActorParams.bStatic = true;
 
-		FPhysicsInterface::CreateActor(ActorParams, Instance->ActorHandle);
+		FPhysicsInterface::CreateActor(ActorParams, Instance->ActorHandle, PrimitiveComp);
 	}
 	else
 	{
-		FPhysicsInterface::CreateActor(ActorParams, Instance->ActorHandle);
+		FPhysicsInterface::CreateActor(ActorParams, Instance->ActorHandle, PrimitiveComp);
 		FPhysicsInterface::SetCcdEnabled_AssumesLocked(Instance->ActorHandle, Instance->bUseCCD);
 		FPhysicsInterface::SetMACDEnabled_AssumesLocked(Instance->ActorHandle, Instance->IsUsingMACD());
 		FPhysicsInterface::SetIsKinematic_AssumesLocked(Instance->ActorHandle, !Instance->ShouldInstanceSimulatingPhysics());
diff --git a/Engine/Source/Runtime/Engine/Private/PhysicsEngine/ClusterUnionComponent.cpp b/Engine/Source/Runtime/Engine/Private/PhysicsEngine/ClusterUnionComponent.cpp
--- a/Engine/Source/Runtime/Engine/Private/PhysicsEngine/ClusterUnionComponent.cpp
+++ b/Engine/Source/Runtime/Engine/Private/PhysicsEngine/ClusterUnionComponent.cpp
@@ -1011,6 +1011,7 @@ void UClusterUnionComponent::SyncClusterUnionFromProxy(const FTransform& NewTran
 					}
 				}
 
+				ensure(ChildData.BoneId != INDEX_NONE);
 				Chaos::FPhysicsObjectHandle Handle = Component->GetPhysicsObjectById(ChildData.BoneId);
 				Chaos::FPBDRigidParticle* Particle = Interface->GetRigidParticle(Handle);
 
diff --git a/Engine/Source/Runtime/Engine/Private/SkeletalMeshComponentPhysics.cpp b/Engine/Source/Runtime/Engine/Private/SkeletalMeshComponentPhysics.cpp
--- a/Engine/Source/Runtime/Engine/Private/SkeletalMeshComponentPhysics.cpp
+++ b/Engine/Source/Runtime/Engine/Private/SkeletalMeshComponentPhysics.cpp
@@ -1622,6 +1622,11 @@ bool USkeletalMeshComponent::ShouldCreatePhysicsState() const
 	return bShouldCreatePhysicsState;
 }
 
+bool USkeletalMeshComponent::HasValidPhysicsState() const
+{
+	return !Bodies.IsEmpty();
+}
+
 void USkeletalMeshComponent::OnCreatePhysicsState()
 {
 	// Init physics
diff --git a/Engine/Source/Runtime/Experimental/Chaos/Private/Chaos/PBDRigidClustering.cpp b/Engine/Source/Runtime/Experimental/Chaos/Private/Chaos/PBDRigidClustering.cpp
--- a/Engine/Source/Runtime/Experimental/Chaos/Private/Chaos/PBDRigidClustering.cpp
+++ b/Engine/Source/Runtime/Experimental/Chaos/Private/Chaos/PBDRigidClustering.cpp
@@ -75,6 +75,7 @@ namespace Chaos
 	namespace CVars
 	{
 		extern CHAOS_API bool bChaosConvexSimplifyUnion;
+		extern bool bChaosEnableOverrideSingleParticleTypeToCluster;
 	}
 
 	template <typename TProxy=FGeometryCollectionPhysicsProxy>
@@ -1101,6 +1102,13 @@ namespace Chaos
 			return ActivatedChildren;
 		}
 
+		if (CVars::bChaosEnableOverrideSingleParticleTypeToCluster)
+		{
+			// early out so components that are added from static and skeletal meshes that rely on SingleParticlePhysicsProxy will not fall to pieces immediately!
+			// TODO: This means they are currently unbreakable which is currently ok, but we may need them to have damage in the future
+			return ActivatedChildren;
+		}
+
 		// gather propagation information from the parent proxy
 		bool bUseDamagePropagation = false;
 		float BreakDamagePropagationFactor = 0.0f;
diff --git a/Engine/Source/Runtime/Experimental/Chaos/Private/PBDRigidsSolver.cpp b/Engine/Source/Runtime/Experimental/Chaos/Private/PBDRigidsSolver.cpp
--- a/Engine/Source/Runtime/Experimental/Chaos/Private/PBDRigidsSolver.cpp
+++ b/Engine/Source/Runtime/Experimental/Chaos/Private/PBDRigidsSolver.cpp
@@ -414,6 +414,8 @@ namespace Chaos
 		int32 ChaosVisualDebuggerEnable = 1;
 		FAutoConsoleVariableRef CVarChaosVisualDebuggerEnable(TEXT("p.Chaos.VisualDebuggerEnable"), ChaosVisualDebuggerEnable, TEXT("Enable/Disable pushing/saving data to the visual debugger"));
 
+		bool bChaosEnableOverrideSingleParticleTypeToCluster = false;
+		FAutoConsoleVariableRef CVarChaosEnableOverrideSingleParticleTypeToCluster(TEXT("p.Chaos.EnableOverrideSingleParticleTypeToCluster"), bChaosEnableOverrideSingleParticleTypeToCluster, TEXT("When enabled particles are promoted to cluster type so they may be used with cluster unions"));
 
 		// Enable a couple bug fixes with temporary roll-back just in case
 		bool bRemoveParticleFromMovingKinematicsOnDisable = true;
@@ -1598,7 +1600,19 @@ namespace Chaos
 							{
 								case EParticleType::Static: return Particles.CreateStaticParticles(1, UniqueIdx)[0];
 								case EParticleType::Kinematic: return Particles.CreateKinematicParticles(1, UniqueIdx)[0];
-								case EParticleType::Rigid: return Particles.CreateDynamicParticles(1, UniqueIdx)[0];
+
+								case EParticleType::Rigid: 
+								{
+									if (CVars::bChaosEnableOverrideSingleParticleTypeToCluster)
+									{
+										return Particles.CreateClusteredParticles(1, UniqueIdx)[0];
+									}
+									else
+									{
+										return Particles.CreateDynamicParticles(1, UniqueIdx)[0];
+									}
+								}
+								
 								default: check(false); return nullptr;
 							}
 						});
@@ -2122,35 +2136,19 @@ namespace Chaos
 
 			for (TPBDRigidParticleHandleImp<FReal, 3, false>& DirtyParticle : DirtyParticles)
 			{
+				const EParticleType ParticleType = DirtyParticle.GetParticleType();
+				if ((ParticleType == EParticleType::Kinematic) || (ParticleType == EParticleType::Static))
+				{
+					ensure(false);
+				}
+
 				if(IPhysicsProxyBase* Proxy = DirtyParticle.Handle()->PhysicsProxy())
 				{
-					switch(DirtyParticle.GetParticleType())
+					switch (Proxy->GetType())
 					{
-						case EParticleType::Rigid:
+						case EPhysicsProxyType::GeometryCollectionType:
+						case EPhysicsProxyType::ClusterUnionProxy:
 						{
-							if(!bIsResim || DirtyParticle.SyncState() == ESyncState::HardDesync)
-							{
-								// Although per-particle we can control the syncing of target positions (see ShouldUpdateFromSimulation)
-								// we cannot avoid marking kinematics as dirty (unless the global config Chaos::SyncKinematicOnGameThread
-								// forces it) because we always need the correct velocities/dynamics to be synced back from the kinematic
-								// target.
-								// FSingleParticlePhysicsProxy::PullFromPhysicsState will use the proper checks to see whether it needs
-								// to sync the particle positions when the dirty particle is processed.
-								if(!(Chaos::SyncKinematicOnGameThread == 0 && DirtyParticle.ObjectState() == EObjectStateType::Kinematic))
-								{
-									ActiveRigid.AddUnique((FSingleParticlePhysicsProxy*)Proxy);
-								}
-							}
-							break;
-						}
-						case EParticleType::Kinematic:
-						case EParticleType::Static:
-							ensure(false);
-							break;
-						case EParticleType::GeometryCollection:
-							ActiveGC.AddUnique((FGeometryCollectionPhysicsProxy*)(Proxy));
-							break;
-						case EParticleType::Clustered:
 							if (auto ClusterParticle = DirtyParticle.CastToClustered())
 							{
 								if (ClusterParticle->InternalCluster())
@@ -2168,17 +2166,17 @@ namespace Chaos
 									// otherwise we should not need a full update (if the cluster is still we only
 									// need to consider de-clustered objects which should be in other parts of the 
 									// dirty view)
-									if(bActivateChildren)
+									if (bActivateChildren)
 									{
 										const TSet<IPhysicsProxyBase*> Proxies = ClusterParticle->PhysicsProxies();
-										for(IPhysicsProxyBase* ClusterProxy : Proxies)
+										for (IPhysicsProxyBase* ClusterProxy : Proxies)
 										{
-											if(!ClusterProxy)
+											if (!ClusterProxy)
 											{
 												continue;
 											}
 
-											switch(ClusterProxy->GetType())
+											switch (ClusterProxy->GetType())
 											{
 											case EPhysicsProxyType::SingleParticleProxy:
 												ActiveRigid.AddUnique((FSingleParticlePhysicsProxy*)ClusterProxy);
@@ -2195,14 +2193,50 @@ namespace Chaos
 								}
 								else
 								{
-									ActiveGC.AddUnique((FGeometryCollectionPhysicsProxy*)(Proxy));
+									switch (Proxy->GetType())
+									{
+									case EPhysicsProxyType::SingleParticleProxy:
+										ActiveRigid.AddUnique((FSingleParticlePhysicsProxy*)Proxy);
+										break;
+									case EPhysicsProxyType::GeometryCollectionType:
+										ActiveGC.AddUnique((FGeometryCollectionPhysicsProxy*)(Proxy));
+										break;
+									default:
+										ensure(false);
+										break;
+									}
 								}
 							}
-							break;
+							else
+							{
+								ensure(false);
+							}
+						}
+						break;
+						case EPhysicsProxyType::SingleParticleProxy:
+						{
+							if (!bIsResim || DirtyParticle.SyncState() == ESyncState::HardDesync)
+							{
+								// Although per-particle we can control the syncing of target positions (see ShouldUpdateFromSimulation)
+								// we cannot avoid marking kinematics as dirty (unless the global config Chaos::SyncKinematicOnGameThread
+								// forces it) because we always need the correct velocities/dynamics to be synced back from the kinematic
+								// target.
+								// FSingleParticlePhysicsProxy::PullFromPhysicsState will use the proper checks to see whether it needs
+								// to sync the particle positions when the dirty particle is processed.
+								if (!(Chaos::SyncKinematicOnGameThread == 0 && DirtyParticle.ObjectState() == EObjectStateType::Kinematic))
+								{
+									ActiveRigid.AddUnique((FSingleParticlePhysicsProxy*)Proxy);
+								}
+							}
+						}
+						break;
+
 						default:
 							check(false);
+							break;
 					}
-				}	
+
+				}
 			}
 		}
 
diff --git a/Engine/Source/Runtime/Experimental/Chaos/Private/PhysicsProxy/ClusterUnionPhysicsProxy.cpp b/Engine/Source/Runtime/Experimental/Chaos/Private/PhysicsProxy/ClusterUnionPhysicsProxy.cpp
--- a/Engine/Source/Runtime/Experimental/Chaos/Private/PhysicsProxy/ClusterUnionPhysicsProxy.cpp
+++ b/Engine/Source/Runtime/Experimental/Chaos/Private/PhysicsProxy/ClusterUnionPhysicsProxy.cpp
@@ -15,6 +15,11 @@ namespace Chaos
 	static int32 GShallowCopyClusterUnionGeometryOnUpdate = 1;
 	FAutoConsoleVariableRef CVar_ShallowCopyClusterUnionGeometryOnUpdate(TEXT("p.ShallowCopyOnClusterUnionUpdate"), GShallowCopyClusterUnionGeometryOnUpdate, TEXT("If 1, shallow copy the root union geometry of a cluster union when its geometry updates, otherwise deep copy the geometry hierarchy"));
 
+	namespace CVars
+	{
+		extern bool bChaosEnableOverrideSingleParticleTypeToCluster;
+	}
+
 	namespace
 	{
 		FPBDRigidsEvolutionGBF* GetEvolution(FClusterUnionPhysicsProxy* Proxy)
@@ -662,6 +667,28 @@ namespace Chaos
 			}
 		}
 		
+		// Note: Is their a better way to achieve this? 
+		// Don't want to iterate over all children if we don't need to so using flag
+		if (CVars::bChaosEnableOverrideSingleParticleTypeToCluster)
+		{
+			for (const FDirtyClusterUnionParticleData& InData : CurrentPullData.ChildParticles)
+			{
+				// SingleParticleProxies will not render in the correct location if they are still part on an intact cluster
+				// so update the transform of these 'kinematic' child particles
+				if (InData.Proxy && InData.Proxy->GetType() == EPhysicsProxyType::SingleParticleProxy)
+				{
+					if (Chaos::FSingleParticlePhysicsProxy* SingleParticleProxy = static_cast<Chaos::FSingleParticlePhysicsProxy*>(InData.Proxy))
+					{
+						Chaos::FRigidBodyHandle_External& Body_External = SingleParticleProxy->GetGameThreadAPI();
+						FTransform ClusterTransform(Particle_External->R(), Particle_External->X());
+						FTransform ChildWorldTransform = InData.ChildToParent * ClusterTransform;
+						Body_External.SetX(ChildWorldTransform.GetLocation());
+						Body_External.SetR(ChildWorldTransform.GetRotation());
+					}
+				}
+			}
+		}
+			
 		Particle_External->UpdateShapeBounds();
 		return true;
 	}
@@ -707,6 +734,11 @@ namespace Chaos
 						FGeometryCollectionPhysicsProxy* GCProxy = static_cast<FGeometryCollectionPhysicsProxy*>(Proxy);
 						Data.BoneId = GCProxy->GetTransformGroupIndexFromHandle(Particle);
 					}
+					else if (Proxy->GetType() == EPhysicsProxyType::SingleParticleProxy)
+					{
+						// There's no way to reconstitute the BoneID! BoneID is used to match the particle on the client, but this makes no sense when dealing with SingleParticle
+						Data.BoneId = BufferData.ChildParticles.Num(); // This could be wrong, but seems to work for now!
+					}
 
 					BufferData.ChildParticles.Add(Data);
 				}
diff --git a/Engine/Source/Runtime/Experimental/Chaos/Public/PhysicsProxy/SingleParticlePhysicsProxy.h b/Engine/Source/Runtime/Experimental/Chaos/Public/PhysicsProxy/SingleParticlePhysicsProxy.h
--- a/Engine/Source/Runtime/Experimental/Chaos/Public/PhysicsProxy/SingleParticlePhysicsProxy.h
+++ b/Engine/Source/Runtime/Experimental/Chaos/Public/PhysicsProxy/SingleParticlePhysicsProxy.h
@@ -60,7 +60,7 @@ public:
 	using PARTICLE_TYPE = FGeometryParticle;
 	using FParticleHandle = FGeometryParticleHandle;
 
-	static CHAOS_API FSingleParticlePhysicsProxy* Create(TUniquePtr<FGeometryParticle>&& Particle);
+	static CHAOS_API FSingleParticlePhysicsProxy* Create(TUniquePtr<FGeometryParticle>&& Particle, UObject* InOwner);
 
 	FSingleParticlePhysicsProxy() = delete;
 	FSingleParticlePhysicsProxy(const FSingleParticlePhysicsProxy&) = delete;
@@ -1587,10 +1587,10 @@ public:
 
 static_assert(sizeof(FRigidBodyHandle_Internal) == sizeof(FSingleParticlePhysicsProxy), "Derived types only used to constrain API, all data lives in base class ");
 
-inline FSingleParticlePhysicsProxy* FSingleParticlePhysicsProxy::Create(TUniquePtr<FGeometryParticle>&& Particle)
+inline FSingleParticlePhysicsProxy* FSingleParticlePhysicsProxy::Create(TUniquePtr<FGeometryParticle>&& Particle, UObject* InOwner = nullptr)
 {
 	ensure(Particle->GetProxy() == nullptr);	//not already owned by another proxy. TODO: use TUniquePtr
-	auto Proxy = new FSingleParticlePhysicsProxy(MoveTemp(Particle), nullptr);
+	auto Proxy = new FSingleParticlePhysicsProxy(MoveTemp(Particle), nullptr, InOwner);
 	return Proxy;
 }
 }
diff --git a/Engine/Source/Runtime/PhysicsCore/Private/ChaosEngineInterface.cpp b/Engine/Source/Runtime/PhysicsCore/Private/ChaosEngineInterface.cpp
--- a/Engine/Source/Runtime/PhysicsCore/Private/ChaosEngineInterface.cpp
+++ b/Engine/Source/Runtime/PhysicsCore/Private/ChaosEngineInterface.cpp
@@ -2121,7 +2121,7 @@ int32 FChaosEngineInterface::GetAllShapes_AssumedLocked(const FPhysicsActorHandl
 	return GetAllShapesInternalImp_AssumedLocked(InActorHandle,OutShapes);
 }
 
-void FChaosEngineInterface::CreateActor(const FActorCreationParams& InParams,FPhysicsActorHandle& Handle)
+void FChaosEngineInterface::CreateActor(const FActorCreationParams& InParams,FPhysicsActorHandle& Handle, UObject* InOwner)
 {
 	LLM_SCOPE(ELLMTag::ChaosActor);
 	using namespace Chaos;
@@ -2199,7 +2199,7 @@ void FChaosEngineInterface::CreateActor(const FActorCreationParams& InParams,FPh
 		Particle->SetSpatialIdx(SpatialIndex);
 	}
 
-	Handle = Chaos::FSingleParticlePhysicsProxy::Create(MoveTemp(Particle));
+	Handle = Chaos::FSingleParticlePhysicsProxy::Create(MoveTemp(Particle), InOwner);
 	Chaos::FRigidBodyHandle_External& Body_External = Handle->GetGameThreadAPI();
 
 	// Set up the new particle's game-thread data. This will be sent to physics-thread when
diff --git a/Engine/Source/Runtime/PhysicsCore/Public/Chaos/ChaosEngineInterface.h b/Engine/Source/Runtime/PhysicsCore/Public/Chaos/ChaosEngineInterface.h
--- a/Engine/Source/Runtime/PhysicsCore/Public/Chaos/ChaosEngineInterface.h
+++ b/Engine/Source/Runtime/PhysicsCore/Public/Chaos/ChaosEngineInterface.h
@@ -343,7 +343,7 @@ class FChaosEngineInterface
 public:
 	virtual ~FChaosEngineInterface() = default;
 
-	static PHYSICSCORE_API void CreateActor(const FActorCreationParams& InParams,FPhysicsActorHandle& Handle);
+	static PHYSICSCORE_API void CreateActor(const FActorCreationParams& InParams,FPhysicsActorHandle& Handle, UObject* InOwner = nullptr);
 	static PHYSICSCORE_API void ReleaseActor(FPhysicsActorHandle& InActorReference,FChaosScene* InScene = nullptr,bool bNeverDeferRelease=false);
 
 	static bool IsValid(const FPhysicsActorHandle& Handle) { return Handle != nullptr; }
