Edge case crash bug fix, just removing offending code, not sure why it's even there

diff --git a/Engine/Plugins/Experimental/ChaosModularVehicle/Source/ChaosModularVehicleEngine/Private/ChaosModularVehicle/ChaosSimModuleManager.cpp b/Engine/Plugins/Experimental/ChaosModularVehicle/Source/ChaosModularVehicleEngine/Private/ChaosModularVehicle/ChaosSimModuleManager.cpp
--- a/Engine/Plugins/Experimental/ChaosModularVehicle/Source/ChaosModularVehicleEngine/Private/ChaosModularVehicle/ChaosSimModuleManager.cpp
+++ b/Engine/Plugins/Experimental/ChaosModularVehicle/Source/ChaosModularVehicleEngine/Private/ChaosModularVehicle/ChaosSimModuleManager.cpp
@@ -97,11 +97,6 @@ void FChaosSimModuleManager::DetachFromPhysScene(FPhysScene* PhysScene)
 		UnregisterCallbacks();
 	}
 
-	if (PhysScene->GetOwningWorld())
-	{
-		PhysScene->GetOwningWorld()->OnWorldBeginPlay.RemoveAll(this);
-	}
-
 	FChaosSimModuleManager::SceneToModuleManagerMap.Remove(PhysScene);
 }
 
diff --git a/Engine/Plugins/Experimental/ChaosVehiclesPlugin/Source/ChaosVehicles/Private/ChaosVehicleManager.cpp b/Engine/Plugins/Experimental/ChaosVehiclesPlugin/Source/ChaosVehicles/Private/ChaosVehicleManager.cpp
--- a/Engine/Plugins/Experimental/ChaosVehiclesPlugin/Source/ChaosVehicles/Private/ChaosVehicleManager.cpp
+++ b/Engine/Plugins/Experimental/ChaosVehiclesPlugin/Source/ChaosVehicles/Private/ChaosVehicleManager.cpp
@@ -101,11 +101,6 @@ void FChaosVehicleManager::DetachFromPhysScene(FPhysScene* PhysScene)
 		UnregisterCallbacks();
 	}
 
-	if (PhysScene->GetOwningWorld())
-	{
-		PhysScene->GetOwningWorld()->OnWorldBeginPlay.RemoveAll(this);
-	}
-
 	FChaosVehicleManager::SceneToVehicleManagerMap.Remove(PhysScene);
 }
 
