Replaced `IsValid(this)` under the rest of Engine/.

diff --git a/Engine/Plugins/2D/Paper2D/Source/Paper2D/Private/PaperCharacter.cpp b/Engine/Plugins/2D/Paper2D/Source/Paper2D/Private/PaperCharacter.cpp
--- a/Engine/Plugins/2D/Paper2D/Source/Paper2D/Private/PaperCharacter.cpp
+++ b/Engine/Plugins/2D/Paper2D/Source/Paper2D/Private/PaperCharacter.cpp
@@ -35,7 +35,7 @@ void APaperCharacter::PostInitializeComponents()
 {
 	Super::PostInitializeComponents();
 
-	if (IsValid(this))
+	if (IsValidChecked(this))
 	{
 		if (Sprite)
 		{
diff --git a/Engine/Plugins/Animation/ControlRig/Source/ControlRig/Private/Rigs/RigHierarchy.cpp b/Engine/Plugins/Animation/ControlRig/Source/ControlRig/Private/Rigs/RigHierarchy.cpp
--- a/Engine/Plugins/Animation/ControlRig/Source/ControlRig/Private/Rigs/RigHierarchy.cpp
+++ b/Engine/Plugins/Animation/ControlRig/Source/ControlRig/Private/Rigs/RigHierarchy.cpp
@@ -4283,7 +4283,7 @@ void URigHierarchy::Notify(ERigHierarchyNotification InNotifType, const FRigNoti
 		return;
 	}
 
-	if (!IsValid(this))
+	if (!IsValidChecked(this))
 	{
 		return;
 	}
diff --git a/Engine/Plugins/Animation/ControlRig/Source/ControlRig/Private/Rigs/RigHierarchyController.cpp b/Engine/Plugins/Animation/ControlRig/Source/ControlRig/Private/Rigs/RigHierarchyController.cpp
--- a/Engine/Plugins/Animation/ControlRig/Source/ControlRig/Private/Rigs/RigHierarchyController.cpp
+++ b/Engine/Plugins/Animation/ControlRig/Source/ControlRig/Private/Rigs/RigHierarchyController.cpp
@@ -2577,7 +2577,9 @@ void URigHierarchyController::HandleHierarchyModified(ERigHierarchyNotification
 
 bool URigHierarchyController::IsValid() const
 {
-	if(::IsValid(this))
+	PRAGMA_DISABLE_DEPRECATION_WARNINGS
+	if(IsThisNotNull(this, "URigHierarchyController::IsValid") && IsValidChecked(this))
+	PRAGMA_ENABLE_DEPRECATION_WARNINGS
 	{
 		return ::IsValid(GetHierarchy());
 	}
diff --git a/Engine/Plugins/Animation/ControlRig/Source/ControlRigDeveloper/Private/ControlRigBlueprint.cpp b/Engine/Plugins/Animation/ControlRig/Source/ControlRigDeveloper/Private/ControlRigBlueprint.cpp
--- a/Engine/Plugins/Animation/ControlRig/Source/ControlRigDeveloper/Private/ControlRigBlueprint.cpp
+++ b/Engine/Plugins/Animation/ControlRig/Source/ControlRigDeveloper/Private/ControlRigBlueprint.cpp
@@ -855,7 +855,7 @@ void UControlRigBlueprint::Serialize(FArchive& Ar)
 {
 	UE_RIGVM_ARCHIVETRACE_SCOPE(Ar, FString::Printf(TEXT("UControlRigBlueprint(%s)"), *GetName()));
 
-	if(IsValid(this))
+	if(IsValidChecked(this))
 	{
 		RigVMClient.SetOuterClientHost(this, GET_MEMBER_NAME_CHECKED(UControlRigBlueprint, RigVMClient));
 		ModularRigModel.SetOuterClientHost(this);
diff --git a/Engine/Plugins/Experimental/Avalanche/Source/AvalancheShapes/Private/DynamicMeshes/AvaShapeDynMeshBase.cpp b/Engine/Plugins/Experimental/Avalanche/Source/AvalancheShapes/Private/DynamicMeshes/AvaShapeDynMeshBase.cpp
--- a/Engine/Plugins/Experimental/Avalanche/Source/AvalancheShapes/Private/DynamicMeshes/AvaShapeDynMeshBase.cpp
+++ b/Engine/Plugins/Experimental/Avalanche/Source/AvalancheShapes/Private/DynamicMeshes/AvaShapeDynMeshBase.cpp
@@ -269,7 +269,9 @@ void UAvaShapeDynamicMeshBase::ToggleGizmo_Implementation(
 
 void UAvaShapeDynamicMeshBase::UpdateDynamicMesh()
 {
-	if (IsValid(this) && !HasAnyFlags(RF_BeginDestroyed | RF_FinishDestroyed | RF_ClassDefaultObject))
+	PRAGMA_DISABLE_DEPRECATION_WARNINGS
+	if (IsThisNotNull(this, "UAvaShapeDynamicMeshBase::UpdateDynamicMesh") && IsValidChecked(this) && !HasAnyFlags(RF_BeginDestroyed | RF_FinishDestroyed | RF_ClassDefaultObject))
+	PRAGMA_ENABLE_DEPRECATION_WARNINGS
 	{
 		RestoreCachedMesh();
 		CachedMesh.Reset();
diff --git a/Engine/Plugins/Experimental/ChaosModularVehicle/Source/ChaosModularVehicleEngine/Private/ChaosModularVehicle/ModularVehicleBaseComponent.cpp b/Engine/Plugins/Experimental/ChaosModularVehicle/Source/ChaosModularVehicleEngine/Private/ChaosModularVehicle/ModularVehicleBaseComponent.cpp
--- a/Engine/Plugins/Experimental/ChaosModularVehicle/Source/ChaosModularVehicleEngine/Private/ChaosModularVehicle/ModularVehicleBaseComponent.cpp
+++ b/Engine/Plugins/Experimental/ChaosModularVehicle/Source/ChaosModularVehicleEngine/Private/ChaosModularVehicle/ModularVehicleBaseComponent.cpp
@@ -1061,57 +1061,57 @@ void UModularVehicleBaseComponent::ActionTreeUpdates(Chaos::FSimTreeUpdates* Nex
 	{
 		if (Chaos::FPBDRigidsSolver* Solver = Proxy->GetSolver<Chaos::FPBDRigidsSolver>())
 		{
-Solver->EnqueueCommandImmediate([Proxy, this, NextTreeUpdates = *NextTreeUpdates]() mutable
-	{
-		if (IsValid(this) && bPhysicsStateCreated && VehicleSimulationPT)
-		{
-			TUniquePtr<Chaos::FSimModuleTree>& SimModuleTree = VehicleSimulationPT->AccessSimComponentTree();
-			if(SimModuleTree.IsValid())
+			Solver->EnqueueCommandImmediate([Proxy, this, NextTreeUpdates = *NextTreeUpdates]() mutable
 			{
-				SimModuleTree->AppendTreeUpdates(NextTreeUpdates);
-				FModularVehicleBuilder::FixupTreeLinks(SimModuleTree);
-
-#if !(UE_BUILD_SHIPPING || UE_BUILD_TEST)
-				if (bModularVehicle_DumpModuleTreeStructure_Enabled)
+				if (IsValidChecked(this) && bPhysicsStateCreated && VehicleSimulationPT)
 				{
-					UE_LOG(LogModularBase, Warning, TEXT("SimTreeModules:") );
-					for (int I = 0; I < SimModuleTree->GetNumNodes(); I++)
+					TUniquePtr<Chaos::FSimModuleTree>& SimModuleTree = VehicleSimulationPT->AccessSimComponentTree();
+					if(SimModuleTree.IsValid())
 					{
-						if (Chaos::ISimulationModuleBase* Module = SimModuleTree->GetNode(I).SimModule)
+						SimModuleTree->AppendTreeUpdates(NextTreeUpdates);
+						FModularVehicleBuilder::FixupTreeLinks(SimModuleTree);
+
+#if !(UE_BUILD_SHIPPING || UE_BUILD_TEST)
+						if (bModularVehicle_DumpModuleTreeStructure_Enabled)
 						{
-							FString String;
-							Module->GetDebugString(String);
-							UE_LOG(LogModularBase, Warning, TEXT("..%s"), *String);
+							UE_LOG(LogModularBase, Warning, TEXT("SimTreeModules:") );
+							for (int I = 0; I < SimModuleTree->GetNumNodes(); I++)
+							{
+								if (Chaos::ISimulationModuleBase* Module = SimModuleTree->GetNode(I).SimModule)
+								{
+									FString String;
+									Module->GetDebugString(String);
+									UE_LOG(LogModularBase, Warning, TEXT("..%s"), *String);
+								}
+							}
 						}
-					}
-				}
 #endif // !(UE_BUILD_SHIPPING || UE_BUILD_TEST)
 
-				// Network replication data needs to be updated, this is currently studily slow
-				if (NetworkPhysicsComponent)
-				{
-					TSharedPtr<Chaos::FBaseRewindHistory>& History = NetworkPhysicsComponent->GetStateHistory_Internal();
-					Chaos::TDataRewindHistory<FNetworkModularVehicleStates>* StateHistory = static_cast<Chaos::TDataRewindHistory<FNetworkModularVehicleStates>*>(History.Get());
-					if (StateHistory)
-					{
-						// #TODO: we are rebuilding from scratch every time there is a single change, there must be a better way!
-						// not sure of it is safe to update the data at this time?
-						for (int I = 0; I < StateHistory->GetDataHistory().Num(); I++)
+						// Network replication data needs to be updated, this is currently studily slow
+						if (NetworkPhysicsComponent)
 						{
-							FNetworkModularVehicleStates& State = StateHistory->GetDataHistory()[I];
-							State.ModuleData.Empty();
-							TUniquePtr<Chaos::FSimModuleTree>& InnerSimModuleTree = VehicleSimulationPT->AccessSimComponentTree();
-							if(InnerSimModuleTree.IsValid())
+							TSharedPtr<Chaos::FBaseRewindHistory>& History = NetworkPhysicsComponent->GetStateHistory_Internal();
+							Chaos::TDataRewindHistory<FNetworkModularVehicleStates>* StateHistory = static_cast<Chaos::TDataRewindHistory<FNetworkModularVehicleStates>*>(History.Get());
+							if (StateHistory)
 							{
-								InnerSimModuleTree->GenerateReplicationStructure(State.ModuleData);
+								// #TODO: we are rebuilding from scratch every time there is a single change, there must be a better way!
+								// not sure of it is safe to update the data at this time?
+								for (int I = 0; I < StateHistory->GetDataHistory().Num(); I++)
+								{
+									FNetworkModularVehicleStates& State = StateHistory->GetDataHistory()[I];
+									State.ModuleData.Empty();
+									TUniquePtr<Chaos::FSimModuleTree>& InnerSimModuleTree = VehicleSimulationPT->AccessSimComponentTree();
+									if(InnerSimModuleTree.IsValid())
+									{
+										InnerSimModuleTree->GenerateReplicationStructure(State.ModuleData);
+									}
+								}
 							}
 						}
 					}
 				}
-			}
+			});
 		}
-	});
-}
 	}
 }
 
@@ -1244,7 +1244,7 @@ void UModularVehicleBaseComponent::RemoveComponentFromSimulationImpl(UPrimitiveC
 		Chaos::FPBDRigidsSolver* Solver = Proxy->GetSolver<Chaos::FPBDRigidsSolver>();
 		Solver->EnqueueCommandImmediate([Proxy, this, LatestTreeUpdates = LatestTreeUpdates]() mutable
 			{
-				if (IsValid(this) && bPhysicsStateCreated && VehicleSimulationPT)
+				if (IsValidChecked(this) && bPhysicsStateCreated && VehicleSimulationPT)
 				{
 					TUniquePtr<Chaos::FSimModuleTree>& SimModuleTree = VehicleSimulationPT->AccessSimComponentTree();
 					if(SimModuleTree.IsValid())
diff --git a/Engine/Plugins/Experimental/CineCameraSceneCapture/Source/CineCameraSceneCapture/Private/CineCameraSceneCaptureComponent.cpp b/Engine/Plugins/Experimental/CineCameraSceneCapture/Source/CineCameraSceneCapture/Private/CineCameraSceneCaptureComponent.cpp
--- a/Engine/Plugins/Experimental/CineCameraSceneCapture/Source/CineCameraSceneCapture/Private/CineCameraSceneCaptureComponent.cpp
+++ b/Engine/Plugins/Experimental/CineCameraSceneCapture/Source/CineCameraSceneCapture/Private/CineCameraSceneCaptureComponent.cpp
@@ -274,7 +274,9 @@ void UCineCaptureComponent2D::TickComponent(float DeltaTime, ELevelTick TickType
 
 void UCineCaptureComponent2D::ValidateCineCameraComponent()
 {
-	if (!IsValid(this))
+	PRAGMA_DISABLE_DEPRECATION_WARNINGS
+	if (!IsThisNotNull(this, "UCineCaptureComponent2D::ValidateCineCameraComponent") || !IsValidChecked(this))
+	PRAGMA_ENABLE_DEPRECATION_WARNINGS
 	{
 		return;
 	}
diff --git a/Engine/Plugins/Experimental/ClonerEffector/Source/ClonerEffector/Private/Cloner/CEClonerComponent.cpp b/Engine/Plugins/Experimental/ClonerEffector/Source/ClonerEffector/Private/Cloner/CEClonerComponent.cpp
--- a/Engine/Plugins/Experimental/ClonerEffector/Source/ClonerEffector/Private/Cloner/CEClonerComponent.cpp
+++ b/Engine/Plugins/Experimental/ClonerEffector/Source/ClonerEffector/Private/Cloner/CEClonerComponent.cpp
@@ -1536,7 +1536,9 @@ void UCEClonerComponent::CreateDefaultActorAttached()
 
 void UCEClonerComponent::ConvertToStaticMesh()
 {
-	if (!IsValid(this) || !bEnabled)
+	PRAGMA_DISABLE_DEPRECATION_WARNINGS
+	if (!IsThisNotNull(this, "UCEClonerComponent::ConvertToStaticMesh") || !IsValidChecked(this) || !bEnabled)
+	PRAGMA_ENABLE_DEPRECATION_WARNINGS
 	{
 		return;
 	}
@@ -1558,7 +1560,9 @@ void UCEClonerComponent::ConvertToStaticMesh()
 
 void UCEClonerComponent::ConvertToDynamicMesh()
 {
-	if (!IsValid(this) || !bEnabled)
+	PRAGMA_DISABLE_DEPRECATION_WARNINGS
+	if (!IsThisNotNull(this, "UCEClonerComponent::ConvertToDynamicMesh") || !IsValidChecked(this) || !bEnabled)
+	PRAGMA_ENABLE_DEPRECATION_WARNINGS
 	{
 		return;
 	}
@@ -1580,7 +1584,9 @@ void UCEClonerComponent::ConvertToDynamicMesh()
 
 void UCEClonerComponent::ConvertToStaticMeshes()
 {
-	if (!IsValid(this) || !bEnabled)
+	PRAGMA_DISABLE_DEPRECATION_WARNINGS
+	if (!IsThisNotNull(this, "UCEClonerComponent::ConvertToStaticMeshes") || !IsValidChecked(this) || !bEnabled)
+	PRAGMA_ENABLE_DEPRECATION_WARNINGS
 	{
 		return;
 	}
@@ -1602,7 +1608,9 @@ void UCEClonerComponent::ConvertToStaticMeshes()
 
 void UCEClonerComponent::ConvertToDynamicMeshes()
 {
-	if (!IsValid(this) || !bEnabled)
+	PRAGMA_DISABLE_DEPRECATION_WARNINGS
+	if (!IsThisNotNull(this, "UCEClonerComponent::ConvertToDynamicMeshes") || !IsValidChecked(this) || !bEnabled)
+	PRAGMA_ENABLE_DEPRECATION_WARNINGS
 	{
 		return;
 	}
@@ -1624,7 +1632,9 @@ void UCEClonerComponent::ConvertToDynamicMeshes()
 
 void UCEClonerComponent::ConvertToInstancedStaticMeshes()
 {
-	if (!IsValid(this) || !bEnabled)
+	PRAGMA_DISABLE_DEPRECATION_WARNINGS
+	if (!IsThisNotNull(this, "UCEClonerComponent::ConvertToInstancedStaticMeshes") || !IsValidChecked(this) || !bEnabled)
+	PRAGMA_ENABLE_DEPRECATION_WARNINGS
 	{
 		return;
 	}
diff --git a/Engine/Plugins/Experimental/ClonerEffector/Source/ClonerEffector/Private/Effector/CEEffectorComponent.cpp b/Engine/Plugins/Experimental/ClonerEffector/Source/ClonerEffector/Private/Effector/CEEffectorComponent.cpp
--- a/Engine/Plugins/Experimental/ClonerEffector/Source/ClonerEffector/Private/Effector/CEEffectorComponent.cpp
+++ b/Engine/Plugins/Experimental/ClonerEffector/Source/ClonerEffector/Private/Effector/CEEffectorComponent.cpp
@@ -170,7 +170,9 @@ FCEClonerEffectorChannelData& UCEEffectorComponent::GetChannelData()
 
 void UCEEffectorComponent::RegisterToChannel()
 {
-	if (IsValid(this) && ChannelData.GetIdentifier() == INDEX_NONE)
+	PRAGMA_DISABLE_DEPRECATION_WARNINGS
+	if (IsThisNotNull(this, "UCEEffectorComponent::RegisterToChannel") && IsValidChecked(this) && ChannelData.GetIdentifier() == INDEX_NONE)
+	PRAGMA_ENABLE_DEPRECATION_WARNINGS
 	{
 		// Register this effector to the effector channel
 		if (UCEEffectorSubsystem* EffectorSubsystem = UCEEffectorSubsystem::Get())
diff --git a/Engine/Plugins/Experimental/ColorCorrectRegions/Source/ColorCorrectRegions/Public/ColorCorrectRegion.h b/Engine/Plugins/Experimental/ColorCorrectRegions/Source/ColorCorrectRegions/Public/ColorCorrectRegion.h
--- a/Engine/Plugins/Experimental/ColorCorrectRegions/Source/ColorCorrectRegions/Public/ColorCorrectRegion.h
+++ b/Engine/Plugins/Experimental/ColorCorrectRegions/Source/ColorCorrectRegions/Public/ColorCorrectRegion.h
@@ -248,7 +248,9 @@ protected:
 	template <typename TCCActorType>
 	void ChangeShapeVisibilityForActorTypeInternal(TCCActorType InDesiredType)
 	{
-		if (!IsValid(this) || MeshComponents.Num() != (uint8)TCCActorType::MAX)
+		PRAGMA_DISABLE_DEPRECATION_WARNINGS
+		if (!IsThisNotNull(this, "AColorCorrectRegion::ChangeShapeVisibilityForActorTypeInternal") || !IsValidChecked(this) || MeshComponents.Num() != (uint8)TCCActorType::MAX)
+		PRAGMA_ENABLE_DEPRECATION_WARNINGS
 		{
 			return;
 		}
diff --git a/Engine/Plugins/Experimental/DaySequence/Source/DaySequence/Private/DaySequenceActor.cpp b/Engine/Plugins/Experimental/DaySequence/Source/DaySequence/Private/DaySequenceActor.cpp
--- a/Engine/Plugins/Experimental/DaySequence/Source/DaySequence/Private/DaySequenceActor.cpp
+++ b/Engine/Plugins/Experimental/DaySequence/Source/DaySequence/Private/DaySequenceActor.cpp
@@ -833,7 +833,7 @@ UMovieSceneSubSection* ADaySequenceActor::InitializeDaySequence(const FDaySequen
 
 			const TFunction<void(void)> SetSubTrackMuteStateConditional = [this, SubSection, Conditions = SequenceAsset.Conditions.Conditions]()
 			{
-				if (!IsValid(this) || !IsValid(SubSection))
+				if (!IsValidChecked(this) || !IsValid(SubSection))
 				{
 					return;
 				}
@@ -853,7 +853,7 @@ UMovieSceneSubSection* ADaySequenceActor::InitializeDaySequence(const FDaySequen
 
 			const TFunction<void(void)> SetSubTrackMuteStateUnconditional = [this, SubSection]()
 			{
-				if (!IsValid(this) || !IsValid(SubSection))
+				if (!IsValidChecked(this) || !IsValid(SubSection))
 				{
 					return;
 				}
diff --git a/Engine/Plugins/Experimental/DaySequence/Source/DaySequence/Private/DaySequenceModifierComponent.cpp b/Engine/Plugins/Experimental/DaySequence/Source/DaySequence/Private/DaySequenceModifierComponent.cpp
--- a/Engine/Plugins/Experimental/DaySequence/Source/DaySequence/Private/DaySequenceModifierComponent.cpp
+++ b/Engine/Plugins/Experimental/DaySequence/Source/DaySequence/Private/DaySequenceModifierComponent.cpp
@@ -572,7 +572,7 @@ void UDaySequenceModifierComponent::SetInitialTimeOfDay()
 			{
 				auto WantsStaticTime = [this]() -> bool
 				{
-					return IsValid(this) && bIsEnabled && bIsComponentEnabled &&
+					return IsValidChecked(this) && bIsEnabled && bIsComponentEnabled &&
 						(DayNightCycle == EDayNightCycleMode::FixedTime || DayNightCycle == EDayNightCycleMode::RandomFixedTime);
 				};
 
@@ -813,7 +813,7 @@ UMovieSceneSubSection* UDaySequenceModifierComponent::InitializeDaySequence(cons
 	
 	const TFunction<void(void)> SetSubTrackMuteStateConditional = [this, SubSection, Conditions = Entry.Conditions.Conditions]()
 	{
-		if (!IsValid(this) || !IsValid(SubSection))
+		if (!IsValidChecked(this) || !IsValid(SubSection))
 		{
 			return;
 		}
@@ -834,7 +834,7 @@ UMovieSceneSubSection* UDaySequenceModifierComponent::InitializeDaySequence(cons
 
 	const TFunction<void(void)> SetSubTrackMuteStateUnconditional = [this, SubSection]()
 	{
-		if (!IsValid(this) || !IsValid(SubSection))
+		if (!IsValidChecked(this) || !IsValid(SubSection))
 		{
 			return;
 		}
diff --git a/Engine/Plugins/Experimental/DaySequence/Source/DaySequence/Private/DaySequenceStaticTime.cpp b/Engine/Plugins/Experimental/DaySequence/Source/DaySequence/Private/DaySequenceStaticTime.cpp
--- a/Engine/Plugins/Experimental/DaySequence/Source/DaySequence/Private/DaySequenceStaticTime.cpp
+++ b/Engine/Plugins/Experimental/DaySequence/Source/DaySequence/Private/DaySequenceStaticTime.cpp
@@ -322,7 +322,7 @@ void UDaySequenceStaticTimeContributor::BindToDaySequenceActor(ADaySequenceActor
 	UObject* Outer = GetOuter();
 	auto WantsStaticTime = [this, Outer]()
 	{
-		return IsValid(this) && IsValid(Outer) && bWantsStaticTime;
+		return IsValidChecked(this) && IsValid(Outer) && bWantsStaticTime;
 	};
 
 	auto GetStaticTime = [this, WantsStaticTime](UE::DaySequence::FStaticTimeInfo& OutRequest)
diff --git a/Engine/Plugins/Experimental/DynamicMaterial/Source/DynamicMaterial/Private/Model/DynamicMaterialModel.cpp b/Engine/Plugins/Experimental/DynamicMaterial/Source/DynamicMaterial/Private/Model/DynamicMaterialModel.cpp
--- a/Engine/Plugins/Experimental/DynamicMaterial/Source/DynamicMaterial/Private/Model/DynamicMaterialModel.cpp
+++ b/Engine/Plugins/Experimental/DynamicMaterial/Source/DynamicMaterial/Private/Model/DynamicMaterialModel.cpp
@@ -210,7 +210,7 @@ UDMMaterialValue* UDynamicMaterialModel::GetGlobalParameterValue(FName InName) c
 
 bool UDynamicMaterialModel::IsModelValid() const
 {
-	return (!HasAnyFlags(RF_BeginDestroyed | RF_FinishDestroyed) && IsValid(this));
+	return (!HasAnyFlags(RF_BeginDestroyed | RF_FinishDestroyed) && IsValidChecked(this));
 }
 
 UDMMaterialValue* UDynamicMaterialModel::GetGlobalParameterValueForMaterialProperty(EDMMaterialPropertyType InProperty) const
@@ -570,7 +570,9 @@ void UDynamicMaterialModel::OnTextureUVUpdated(UDMTextureUV* InTextureUV)
 		return;
 	}
 
-	if (IsValid(this) && IsValid(DynamicMaterialInstance))
+	PRAGMA_DISABLE_DEPRECATION_WARNINGS
+	if (IsThisNotNull(this, "UDynamicMaterialModel::OnTextureUVUpdated") && IsValidChecked(this) && IsValid(DynamicMaterialInstance))
+	PRAGMA_ENABLE_DEPRECATION_WARNINGS
 	{
 		InTextureUV->SetMIDParameters(DynamicMaterialInstance);
 	}
diff --git a/Engine/Plugins/Experimental/GeometryMask/Source/GeometryMask/Private/GeometryMaskCanvasResource.cpp b/Engine/Plugins/Experimental/GeometryMask/Source/GeometryMask/Private/GeometryMaskCanvasResource.cpp
--- a/Engine/Plugins/Experimental/GeometryMask/Source/GeometryMask/Private/GeometryMaskCanvasResource.cpp
+++ b/Engine/Plugins/Experimental/GeometryMask/Source/GeometryMask/Private/GeometryMaskCanvasResource.cpp
@@ -57,7 +57,7 @@ UGeometryMaskCanvasResource::UGeometryMaskCanvasResource()
 
 void UGeometryMaskCanvasResource::BeginDestroy()
 {
-	if (IsValid(this)
+	if (IsValidChecked(this)
 		&& CanvasObject
 		&& IsValid(CanvasObject)
 		&& CanvasObject->Canvas)
diff --git a/Engine/Plugins/Experimental/Landmass/Source/Editor/Private/LandmassErosionBrushBase.cpp b/Engine/Plugins/Experimental/Landmass/Source/Editor/Private/LandmassErosionBrushBase.cpp
--- a/Engine/Plugins/Experimental/Landmass/Source/Editor/Private/LandmassErosionBrushBase.cpp
+++ b/Engine/Plugins/Experimental/Landmass/Source/Editor/Private/LandmassErosionBrushBase.cpp
@@ -121,7 +121,7 @@ void ALandmassErosionBrushBase::PostEditChangeProperty(FPropertyChangedEvent& Pr
 
 	// Do a bunch of checks to make sure that we don't try to do anything when the editing is happening inside the blueprint editor.
 	UWorld* World = GetWorld();
-	if (IsTemplate() || !IsValid(this) || !IsValid(World) || World->WorldType != EWorldType::Editor)
+	if (IsTemplate() || !IsValidChecked(this) || !IsValid(World) || World->WorldType != EWorldType::Editor)
 	{
 		return;
 	}
diff --git a/Engine/Plugins/Experimental/Landmass/Source/Editor/Private/LandmassManagerBase.cpp b/Engine/Plugins/Experimental/Landmass/Source/Editor/Private/LandmassManagerBase.cpp
--- a/Engine/Plugins/Experimental/Landmass/Source/Editor/Private/LandmassManagerBase.cpp
+++ b/Engine/Plugins/Experimental/Landmass/Source/Editor/Private/LandmassManagerBase.cpp
@@ -562,7 +562,7 @@ void ALandmassManagerBase::PostEditChangeProperty(FPropertyChangedEvent& Propert
 
 	// Do a bunch of checks to make sure that we don't try to do anything when the editing is happening inside the blueprint editor.
 	UWorld* World = GetWorld();
-	if (IsTemplate() || !IsValid(this) || !IsValid(World) || World->WorldType != EWorldType::Editor)
+	if (IsTemplate() || !IsValidChecked(this) || !IsValid(World) || World->WorldType != EWorldType::Editor)
 	{
 		return;
 	}
diff --git a/Engine/Plugins/Experimental/LandscapePatch/Source/LandscapePatch/Private/LandscapePatchComponent.cpp b/Engine/Plugins/Experimental/LandscapePatch/Source/LandscapePatch/Private/LandscapePatchComponent.cpp
--- a/Engine/Plugins/Experimental/LandscapePatch/Source/LandscapePatch/Private/LandscapePatchComponent.cpp
+++ b/Engine/Plugins/Experimental/LandscapePatch/Source/LandscapePatch/Private/LandscapePatchComponent.cpp
@@ -785,7 +785,7 @@ void ULandscapePatchComponent::RequestLandscapeUpdate(bool bInUserTriggeredUpdat
 bool ULandscapePatchComponent::IsPatchInWorld() const
 {
 	UWorld* World = GetWorld();
-	return !IsTemplate() && IsValid(this) && IsValid(World) && World->WorldType == EWorldType::Editor;
+	return !IsTemplate() && IsValidChecked(this) && IsValid(World) && World->WorldType == EWorldType::Editor;
 }
 
 #if WITH_EDITOR
diff --git a/Engine/Plugins/Experimental/LandscapePatch/Source/LandscapePatch/Private/LandscapePatchManager.cpp b/Engine/Plugins/Experimental/LandscapePatch/Source/LandscapePatch/Private/LandscapePatchManager.cpp
--- a/Engine/Plugins/Experimental/LandscapePatch/Source/LandscapePatch/Private/LandscapePatchManager.cpp
+++ b/Engine/Plugins/Experimental/LandscapePatch/Source/LandscapePatch/Private/LandscapePatchManager.cpp
@@ -593,7 +593,7 @@ void ALandscapePatchManager::PostEditChangeProperty(FPropertyChangedEvent& Prope
 
 	// Do a bunch of checks to make sure that we don't try to do anything when the editing is happening inside the blueprint editor.
 	UWorld* World = GetWorld();
-	if (IsTemplate() || !IsValid(this) || !IsValid(World) || World->WorldType != EWorldType::Editor)
+	if (IsTemplate() || !IsValidChecked(this) || !IsValid(World) || World->WorldType != EWorldType::Editor)
 	{
 		return;
 	}
diff --git a/Engine/Plugins/Experimental/LandscapePatch/Source/LandscapePatch/Private/LandscapeTextureBackedRenderTarget.cpp b/Engine/Plugins/Experimental/LandscapePatch/Source/LandscapePatch/Private/LandscapeTextureBackedRenderTarget.cpp
--- a/Engine/Plugins/Experimental/LandscapePatch/Source/LandscapePatch/Private/LandscapeTextureBackedRenderTarget.cpp
+++ b/Engine/Plugins/Experimental/LandscapePatch/Source/LandscapePatch/Private/LandscapeTextureBackedRenderTarget.cpp
@@ -582,7 +582,7 @@ void ULandscapeHeightTextureBackedRenderTarget::CopyBackFromInternalTexture()
 bool ULandscapeTextureBackedRenderTargetBase::IsCopyingBackAndForthAllowed()
 {
 	UWorld* World = GetWorld();
-	return IsValid(this) && !IsTemplate()
+	return IsValidChecked(this) && !IsTemplate()
 		// Note that having a null world is ok because we get that temporarily while rerunning construction
 		// scripts. However if we do have a world, it should be the normal editor one.
 		&& (!World || (IsValid(World) && World->WorldType == EWorldType::Editor))
diff --git a/Engine/Plugins/Experimental/PostProcessMaterialChainGraph/Source/PostProcessMaterialChainGraph/Private/PPMChainGraphComponent.cpp b/Engine/Plugins/Experimental/PostProcessMaterialChainGraph/Source/PostProcessMaterialChainGraph/Private/PPMChainGraphComponent.cpp
--- a/Engine/Plugins/Experimental/PostProcessMaterialChainGraph/Source/PostProcessMaterialChainGraph/Private/PPMChainGraphComponent.cpp
+++ b/Engine/Plugins/Experimental/PostProcessMaterialChainGraph/Source/PostProcessMaterialChainGraph/Private/PPMChainGraphComponent.cpp
@@ -57,7 +57,7 @@ void UPPMChainGraphExecutorComponent::TickComponent(float DeltaTime, enum ELevel
 			return;
 		}
 
-		bool bIsActiveLoc = IsValid(this)
+		bool bIsActiveLoc = IsValidChecked(this)
 #if WITH_EDITOR
 			&& !GetOwner()->IsHiddenEd()
 #endif 
diff --git a/Engine/Plugins/Experimental/PropertyAnimatorCore/Source/PropertyAnimatorCore/Private/Components/PropertyAnimatorCoreComponent.cpp b/Engine/Plugins/Experimental/PropertyAnimatorCore/Source/PropertyAnimatorCore/Private/Components/PropertyAnimatorCoreComponent.cpp
--- a/Engine/Plugins/Experimental/PropertyAnimatorCore/Source/PropertyAnimatorCore/Private/Components/PropertyAnimatorCoreComponent.cpp
+++ b/Engine/Plugins/Experimental/PropertyAnimatorCore/Source/PropertyAnimatorCore/Private/Components/PropertyAnimatorCoreComponent.cpp
@@ -410,7 +410,7 @@ void UPropertyAnimatorCoreComponent::PostEditUndo()
 {
 	Super::PostEditUndo();
 
-	if (!bRegistered && IsValid(this))
+	if (!bRegistered && IsValidChecked(this))
 	{
 		UPropertyAnimatorCoreBase::OnAnimatorAddedDelegate.Broadcast(this, nullptr, EPropertyAnimatorCoreUpdateEvent::Undo);
 	}
diff --git a/Engine/Plugins/Experimental/UserToolBoxCore/Source/UserToolBoxCore/Private/UTBBaseTab.cpp b/Engine/Plugins/Experimental/UserToolBoxCore/Source/UserToolBoxCore/Private/UTBBaseTab.cpp
--- a/Engine/Plugins/Experimental/UserToolBoxCore/Source/UserToolBoxCore/Private/UTBBaseTab.cpp
+++ b/Engine/Plugins/Experimental/UserToolBoxCore/Source/UserToolBoxCore/Private/UTBBaseTab.cpp
@@ -222,7 +222,9 @@ UUTBTabSection* UUserToolBoxBaseTab::GetPlaceHolderSection()
 
 bool UUserToolBoxBaseTab::ContainsCommand(const UUTBBaseCommand* Command)const
 {
-	if (!IsValid(this))
+	PRAGMA_DISABLE_DEPRECATION_WARNINGS
+	if (!IsThisNotNull(this, "UUserToolBoxBaseTab::ContainsCommand") || !IsValidChecked(this))
+	PRAGMA_ENABLE_DEPRECATION_WARNINGS
 	{
 		return false;
 	}
diff --git a/Engine/Plugins/FX/Niagara/Source/Niagara/Private/NiagaraComponent.cpp b/Engine/Plugins/FX/Niagara/Source/Niagara/Private/NiagaraComponent.cpp
--- a/Engine/Plugins/FX/Niagara/Source/Niagara/Private/NiagaraComponent.cpp
+++ b/Engine/Plugins/FX/Niagara/Source/Niagara/Private/NiagaraComponent.cpp
@@ -1898,7 +1898,9 @@ void UNiagaraComponent::DestroyInstanceNotComponent()
 
 void UNiagaraComponent::OnPooledReuse(UWorld* NewWorld)
 {
-	check(IsValid(this));
+	PRAGMA_DISABLE_DEPRECATION_WARNINGS
+	check(IsThisNotNull(this, "UNiagaraComponent::OnPooledReuse") && IsValidChecked(this));
+	PRAGMA_ENABLE_DEPRECATION_WARNINGS
 	SetUserParametersToDefaultValues();
 
 	//Need to reset the component's visibility in case it's returned to the pool while marked invisible.
diff --git a/Engine/Plugins/Importers/USDImporter/Source/USDStage/Private/USDStageActor.cpp b/Engine/Plugins/Importers/USDImporter/Source/USDStage/Private/USDStageActor.cpp
--- a/Engine/Plugins/Importers/USDImporter/Source/USDStage/Private/USDStageActor.cpp
+++ b/Engine/Plugins/Importers/USDImporter/Source/USDStage/Private/USDStageActor.cpp
@@ -1013,7 +1013,7 @@ AUsdStageActor::AUsdStageActor()
 						&& !RootLayer.FilePath.IsEmpty())
 					{
 						// Other user deleted us
-						if (!IsValid(this))
+						if (!IsValidChecked(this))
 						{
 							Reset();
 						}
diff --git a/Engine/Plugins/NetcodeUnitTest/NetcodeUnitTest/Source/NetcodeUnitTest/Classes/UnitTestManager.h b/Engine/Plugins/NetcodeUnitTest/NetcodeUnitTest/Source/NetcodeUnitTest/Classes/UnitTestManager.h
--- a/Engine/Plugins/NetcodeUnitTest/NetcodeUnitTest/Source/NetcodeUnitTest/Classes/UnitTestManager.h
+++ b/Engine/Plugins/NetcodeUnitTest/NetcodeUnitTest/Source/NetcodeUnitTest/Classes/UnitTestManager.h
@@ -274,7 +274,7 @@ public:
 	virtual bool IsTickable() const override
 	{
 		// @todo #JohnBLowPri: Find out how the CDO is getting registered for ticking - this is odd
-		return IsValid(this) && !GIsServer && !HasAnyFlags(RF_ClassDefaultObject);
+		return IsValidChecked(this) && !GIsServer && !HasAnyFlags(RF_ClassDefaultObject);
 	}
 
 	virtual bool IsTickableWhenPaused() const override
diff --git a/Engine/Plugins/PCG/Source/PCG/Private/PCGComponent.cpp b/Engine/Plugins/PCG/Source/PCG/Private/PCGComponent.cpp
--- a/Engine/Plugins/PCG/Source/PCG/Private/PCGComponent.cpp
+++ b/Engine/Plugins/PCG/Source/PCG/Private/PCGComponent.cpp
@@ -768,13 +768,13 @@ void UPCGComponent::PostProcessGraph(const FBox& InNewBounds, bool bInGenerated,
 		UpdateDynamicTracking();
 #endif // WITH_EDITOR
 
-		if (IsValid(this) && Context)
+		if (IsValidChecked(this) && Context)
 		{
 			CallPostGenerateFunctions(Context);
 		}
 
 		// If Generate function made the component invalid for whatever reason, we re-check if this is valid.
-		if (IsValid(this))
+		if (IsValidChecked(this))
 		{
 			PCGComponent::BroadcastDynamicDelegate(OnPCGGraphGeneratedExternal, this);
 		}
@@ -784,7 +784,7 @@ void UPCGComponent::PostProcessGraph(const FBox& InNewBounds, bool bInGenerated,
 #if WITH_EDITOR
 	const bool bHasGeneratedOutputAfter = GeneratedGraphOutput.TaggedData.Num() > 0;
 
-	if (IsValid(this) && (bHasGeneratedOutputAfter || bHadGeneratedOutputBefore))
+	if (IsValidChecked(this) && (bHasGeneratedOutputAfter || bHadGeneratedOutputBefore))
 	{
 		FProperty* GeneratedOutputProperty = FindFProperty<FProperty>(UPCGComponent::StaticClass(), GET_MEMBER_NAME_CHECKED(UPCGComponent, GeneratedGraphOutput));
 		check(GeneratedOutputProperty);
@@ -795,7 +795,7 @@ void UPCGComponent::PostProcessGraph(const FBox& InNewBounds, bool bInGenerated,
 	StopGenerationInProgress();
 
 	UPCGSubsystem* Subsystem = GetSubsystem();
-	Subsystem->OnPCGComponentGenerationDone.Broadcast(Subsystem, IsValid(this) ? this : nullptr, EPCGGenerationStatus::Completed);
+	Subsystem->OnPCGComponentGenerationDone.Broadcast(Subsystem, IsValidChecked(this) ? this : nullptr, EPCGGenerationStatus::Completed);
 #endif
 }
 
@@ -1725,7 +1725,7 @@ void UPCGComponent::OnUnregister()
 			Subsystem->CancelGeneration(this);
 		}
 
-		Subsystem->OnPCGComponentUnregistered.Broadcast(IsValid(this) ? this : nullptr);
+		Subsystem->OnPCGComponentUnregistered.Broadcast(IsValidChecked(this) ? this : nullptr);
 	}
 #endif // WITH_EDITOR
 
@@ -2146,7 +2146,7 @@ void UPCGComponent::PreEditChange(FProperty* PropertyAboutToChange)
 
 void UPCGComponent::PostEditChangeProperty(FPropertyChangedEvent& PropertyChangedEvent)
 {
-	if (!PropertyChangedEvent.Property || !IsValid(this))
+	if (!PropertyChangedEvent.Property || !IsValidChecked(this))
 	{
 		Super::PostEditChangeProperty(PropertyChangedEvent);
 		return;
@@ -2173,7 +2173,7 @@ void UPCGComponent::PostEditChangeProperty(FPropertyChangedEvent& PropertyChange
 	Super::PostEditChangeProperty(PropertyChangedEvent);
 
 	// BP actors will early out here as construction script will have created a new component.
-	if (!IsValid(this))
+	if (!IsValidChecked(this))
 	{
 		return;
 	}
diff --git a/Engine/Plugins/Protocols/MQTT/Source/MQTTCore/Private/MQTTClientObject.cpp b/Engine/Plugins/Protocols/MQTT/Source/MQTTCore/Private/MQTTClientObject.cpp
--- a/Engine/Plugins/Protocols/MQTT/Source/MQTTCore/Private/MQTTClientObject.cpp
+++ b/Engine/Plugins/Protocols/MQTT/Source/MQTTCore/Private/MQTTClientObject.cpp
@@ -115,7 +115,7 @@ void UMQTTClientObject::Connect(const FOnConnectDelegate& InOnConnect)
 		MQTTClient->Connect()
 		.Next([&](EMQTTConnectReturnCode InReturnCode)
 		{
-			if(IsValid(this))
+			if(IsValidChecked(this))
 			{
 				InOnConnect.ExecuteIfBound(InReturnCode);
 			}
@@ -132,7 +132,7 @@ void UMQTTClientObject::Disconnect(const FOnDisconnectDelegate& InOnDisconnect)
 		MQTTClient->Disconnect()
 		.Next([&]()
 		{
-			if(IsValid(this))
+			if(IsValidChecked(this))
 			{
 				InOnDisconnect.ExecuteIfBound();
 			}
@@ -162,7 +162,7 @@ UMQTTSubscriptionObject* UMQTTClientObject::Subscribe(const FString& InTopic, EM
 		MQTTClient->Subscribe({MakeTuple(InTopic, InQoS)})
 		.Next([&](TArray<FMQTTSubscribeResult> InSubscribeResults)
 		{
-			if(IsValid(this))
+			if(IsValidChecked(this))
 			{
 				if(InSubscribeResults.IsEmpty() || InSubscribeResults[0].ReturnCode == EMQTTSubscribeReturnCode::Failure)
 				{
@@ -208,7 +208,7 @@ TArray<UMQTTSubscriptionObject*> UMQTTClientObject::SubscribeMany(
 		MQTTClient->Subscribe(TopicQoSPairs)
 		.Next([&](TArray<FMQTTSubscribeResult> InSubscribeResults)
 		{
-			if(IsValid(this))
+			if(IsValidChecked(this))
 			{
 				if(InSubscribeResults.IsEmpty() || InSubscribeResults[0].ReturnCode == EMQTTSubscribeReturnCode::Failure)
 				{
diff --git a/Engine/Plugins/Runtime/GameplayAbilities/Source/GameplayAbilities/Private/Abilities/GameplayAbility.cpp b/Engine/Plugins/Runtime/GameplayAbilities/Source/GameplayAbilities/Private/Abilities/GameplayAbility.cpp
--- a/Engine/Plugins/Runtime/GameplayAbilities/Source/GameplayAbilities/Private/Abilities/GameplayAbility.cpp
+++ b/Engine/Plugins/Runtime/GameplayAbilities/Source/GameplayAbilities/Private/Abilities/GameplayAbility.cpp
@@ -134,7 +134,7 @@ int32 UGameplayAbility::GetFunctionCallspace(UFunction* Function, FFrame* Stack)
 bool UGameplayAbility::CallRemoteFunction(UFunction* Function, void* Parameters, FOutParmRec* OutParms, FFrame* Stack)
 {
 	// Make sure we are not invalid when being called
-	if (!IsValid(this))
+	if (!IsValidChecked(this))
 	{
 		ABILITY_LOG(Warning, TEXT("%s:CallRemoteFunction called while 'this' was invalid. Skipping the rest of the code."), *GetPathName());
 		return false;
@@ -215,7 +215,7 @@ bool UGameplayAbility::IsActive() const
 	}
 
 	// NonInstanced and Instanced-Per-Execution abilities are by definition active unless they are pending kill
-	return IsValid(this);
+	return IsValidChecked(this);
 }
 
 bool UGameplayAbility::IsSupportedForNetworking() const
@@ -2203,7 +2203,7 @@ float UGameplayAbility::GetCooldownTimeRemaining() const
 void UGameplayAbility::SetRemoteInstanceHasEnded()
 {
 	// This could potentially happen in shutdown corner cases
-	if (!IsValid(this) || CurrentActorInfo == nullptr)
+	if (!IsValidChecked(this) || CurrentActorInfo == nullptr)
 	{
 		return;
 	}
@@ -2232,7 +2232,7 @@ void UGameplayAbility::SetRemoteInstanceHasEnded()
 void UGameplayAbility::NotifyAvatarDestroyed()
 {
 	// This could potentially happen in shutdown corner cases
-	if (!IsValid(this) || CurrentActorInfo == nullptr)
+	if (!IsValidChecked(this) || CurrentActorInfo == nullptr)
 	{
 		return;
 	}
diff --git a/Engine/Plugins/Runtime/GameplayAbilities/Source/GameplayAbilities/Private/Abilities/Tasks/AbilityTask_NetworkSyncPoint.cpp b/Engine/Plugins/Runtime/GameplayAbilities/Source/GameplayAbilities/Private/Abilities/Tasks/AbilityTask_NetworkSyncPoint.cpp
--- a/Engine/Plugins/Runtime/GameplayAbilities/Source/GameplayAbilities/Private/Abilities/Tasks/AbilityTask_NetworkSyncPoint.cpp
+++ b/Engine/Plugins/Runtime/GameplayAbilities/Source/GameplayAbilities/Private/Abilities/Tasks/AbilityTask_NetworkSyncPoint.cpp
@@ -75,7 +75,9 @@ void UAbilityTask_NetworkSyncPoint::Activate()
 
 void UAbilityTask_NetworkSyncPoint::SyncFinished()
 {
-	if (IsValid(this))
+	PRAGMA_DISABLE_DEPRECATION_WARNINGS
+	if (IsThisNotNull(this, "UAbilityTask_NetworkSyncPoint::SyncFinished") && IsValidChecked(this))
+	PRAGMA_ENABLE_DEPRECATION_WARNINGS
 	{
 		if (ShouldBroadcastAbilityTaskDelegates())
 		{
diff --git a/Engine/Plugins/Runtime/GameplayAbilities/Source/GameplayAbilities/Private/Abilities/Tasks/AbilityTask_WaitTargetData.cpp b/Engine/Plugins/Runtime/GameplayAbilities/Source/GameplayAbilities/Private/Abilities/Tasks/AbilityTask_WaitTargetData.cpp
--- a/Engine/Plugins/Runtime/GameplayAbilities/Source/GameplayAbilities/Private/Abilities/Tasks/AbilityTask_WaitTargetData.cpp
+++ b/Engine/Plugins/Runtime/GameplayAbilities/Source/GameplayAbilities/Private/Abilities/Tasks/AbilityTask_WaitTargetData.cpp
@@ -44,7 +44,7 @@ void UAbilityTask_WaitTargetData::Activate()
 			RegisterTargetDataCallbacks();
 
 
-			if (!IsValid(this))
+			if (!IsValidChecked(this))
 			{
 				return;
 			}
@@ -178,7 +178,7 @@ void UAbilityTask_WaitTargetData::FinalizeTargetActor(AGameplayAbilityTargetActo
 
 void UAbilityTask_WaitTargetData::RegisterTargetDataCallbacks()
 {
-	if (!ensure(IsValid(this)))
+	if (!ensure(IsValidChecked(this)))
 	{
 		return;
 	}
diff --git a/Engine/Plugins/Runtime/GameplayAbilities/Source/GameplayAbilities/Private/GameplayCueNotify_Actor.cpp b/Engine/Plugins/Runtime/GameplayAbilities/Source/GameplayAbilities/Private/GameplayCueNotify_Actor.cpp
--- a/Engine/Plugins/Runtime/GameplayAbilities/Source/GameplayAbilities/Private/GameplayCueNotify_Actor.cpp
+++ b/Engine/Plugins/Runtime/GameplayAbilities/Source/GameplayAbilities/Private/GameplayCueNotify_Actor.cpp
@@ -372,7 +372,7 @@ void AGameplayCueNotify_Actor::GameplayCueFinishedCallback()
 
 bool AGameplayCueNotify_Actor::GameplayCuePendingRemove()
 {
-	return GetLifeSpan() > 0.f || FinishTimerHandle.IsValid() || !IsValid(this);
+	return GetLifeSpan() > 0.f || FinishTimerHandle.IsValid() || !IsValidChecked(this);
 }
 
 bool AGameplayCueNotify_Actor::Recycle()
diff --git a/Engine/Plugins/Runtime/GeometryScripting/Source/GeometryScriptingEditor/Private/GeometryActors/GeneratedDynamicMeshActor.cpp b/Engine/Plugins/Runtime/GeometryScripting/Source/GeometryScriptingEditor/Private/GeometryActors/GeneratedDynamicMeshActor.cpp
--- a/Engine/Plugins/Runtime/GeometryScripting/Source/GeometryScriptingEditor/Private/GeometryActors/GeneratedDynamicMeshActor.cpp
+++ b/Engine/Plugins/Runtime/GeometryScripting/Source/GeometryScriptingEditor/Private/GeometryActors/GeneratedDynamicMeshActor.cpp
@@ -104,7 +104,7 @@ void AGeneratedDynamicMeshActor::PostEditUndo()
 	// Currently (5.1) the checks below will tell us if the undo/redo has destroyed the
 	// Actor, and we assume otherwise it was created
 
-	if ( IsActorBeingDestroyed() || !IsValid(this) )	// equivalent to AActor::IsPendingKillPending()
+	if (IsActorBeingDestroyed() || !IsValidChecked(this))	// equivalent to AActor::IsPendingKillPending()
 	{
 		UnregisterWithGenerationManager();
 	}
@@ -122,7 +122,7 @@ void AGeneratedDynamicMeshActor::PostEditUndo(TSharedPtr<ITransactionObjectAnnot
 	// Currently (5.1) the checks below will tell us if the undo/redo has destroyed the
 	// Actor, and we assume otherwise it was created
 
-	if (IsActorBeingDestroyed() || !IsValid(this))	// equivalent to AActor::IsPendingKillPending()
+	if (IsActorBeingDestroyed() || !IsValidChecked(this))	// equivalent to AActor::IsPendingKillPending()
 	{
 		UnregisterWithGenerationManager();
 	}
diff --git a/Engine/Plugins/Runtime/Metasound/Source/MetasoundEditor/Private/MetasoundEditorGraph.cpp b/Engine/Plugins/Runtime/Metasound/Source/MetasoundEditor/Private/MetasoundEditorGraph.cpp
--- a/Engine/Plugins/Runtime/Metasound/Source/MetasoundEditor/Private/MetasoundEditorGraph.cpp
+++ b/Engine/Plugins/Runtime/Metasound/Source/MetasoundEditor/Private/MetasoundEditorGraph.cpp
@@ -201,7 +201,7 @@ void UMetasoundEditorGraphMember::PostEditUndo()
 
 	Super::PostEditUndo();
 
-	if (IsValid(this))
+	if (IsValidChecked(this))
 	{
 		constexpr bool bPostTransaction = false;
 		SetDataType(TypeName, bPostTransaction);
@@ -684,7 +684,7 @@ void UMetasoundEditorGraphMemberDefaultLiteral::PostEditUndo()
 
 	Super::PostEditUndo();
 
-	if (!IsValid(this))
+	if (!IsValidChecked(this))
 	{
 		return;
 	}
diff --git a/Engine/Plugins/Runtime/RigVM/Source/RigVMDeveloper/Private/RigVMBlueprint.cpp b/Engine/Plugins/Runtime/RigVM/Source/RigVMDeveloper/Private/RigVMBlueprint.cpp
--- a/Engine/Plugins/Runtime/RigVM/Source/RigVMDeveloper/Private/RigVMBlueprint.cpp
+++ b/Engine/Plugins/Runtime/RigVM/Source/RigVMDeveloper/Private/RigVMBlueprint.cpp
@@ -790,7 +790,7 @@ void URigVMBlueprint::Serialize(FArchive& Ar)
 	
 	Ar.UsingCustomVersion(FRigVMObjectVersion::GUID);
 	
-	if(IsValid(this))
+	if(IsValidChecked(this))
 	{
 		RigVMClient.SetOuterClientHost(this, GET_MEMBER_NAME_CHECKED(URigVMBlueprint, RigVMClient));
 	}
diff --git a/Engine/Plugins/Runtime/nDisplay/Source/DisplayCluster/Private/Game/EngineClasses/Scene/DisplayClusterSceneComponentSync.cpp b/Engine/Plugins/Runtime/nDisplay/Source/DisplayCluster/Private/Game/EngineClasses/Scene/DisplayClusterSceneComponentSync.cpp
--- a/Engine/Plugins/Runtime/nDisplay/Source/DisplayCluster/Private/Game/EngineClasses/Scene/DisplayClusterSceneComponentSync.cpp
+++ b/Engine/Plugins/Runtime/nDisplay/Source/DisplayCluster/Private/Game/EngineClasses/Scene/DisplayClusterSceneComponentSync.cpp
@@ -73,7 +73,7 @@ void UDisplayClusterSceneComponentSync::EndPlay(const EEndPlayReason::Type EndPl
 //////////////////////////////////////////////////////////////////////////////////////////////
 bool UDisplayClusterSceneComponentSync::IsActive() const
 {
-	return IsValid(this);
+	return IsValidChecked(this);
 }
 
 
diff --git a/Engine/Plugins/VirtualProduction/RemoteControlProtocolDMX/Source/RemoteControlProtocolDMX/Private/Library/RemoteControlDMXLibraryProxy.cpp b/Engine/Plugins/VirtualProduction/RemoteControlProtocolDMX/Source/RemoteControlProtocolDMX/Private/Library/RemoteControlDMXLibraryProxy.cpp
--- a/Engine/Plugins/VirtualProduction/RemoteControlProtocolDMX/Source/RemoteControlProtocolDMX/Private/Library/RemoteControlDMXLibraryProxy.cpp
+++ b/Engine/Plugins/VirtualProduction/RemoteControlProtocolDMX/Source/RemoteControlProtocolDMX/Private/Library/RemoteControlDMXLibraryProxy.cpp
@@ -89,7 +89,7 @@ void URemoteControlDMXLibraryProxy::Refresh()
 	FCoreDelegates::OnEndFrame.RemoveAll(this);
 
 	// Don't refresh if this object is no longer fully valid
-	if (!IsValid(this))
+	if (!IsValidChecked(this))
 	{
 		return;
 	}
