Modular Vehicle: Fix input running on physics thread for preloaded test data.

diff --git a/Engine/Plugins/Experimental/ChaosModularVehicle/Source/ChaosModularVehicleEngine/Private/ChaosModularVehicle/InputProducer.cpp b/Engine/Plugins/Experimental/ChaosModularVehicle/Source/ChaosModularVehicleEngine/Private/ChaosModularVehicle/InputProducer.cpp
--- a/Engine/Plugins/Experimental/ChaosModularVehicle/Source/ChaosModularVehicleEngine/Private/ChaosModularVehicle/InputProducer.cpp
+++ b/Engine/Plugins/Experimental/ChaosModularVehicle/Source/ChaosModularVehicleEngine/Private/ChaosModularVehicle/InputProducer.cpp
@@ -70,10 +70,10 @@ void UVehiclePlaybackInputProducer::ProduceInput(int32 PhysicsStep, int32 NumSte
 		StartStep = PhysicsStep;
 	}
 
-	int32 UseIndex = PhysicsStep - StartStep;
-	if (UseIndex >=0 && UseIndex <= BufferLength)
+	int32 UseIndex = PhysicsStep - StartStep + NumSteps - 1;
+	if (UseIndex >=0 && UseIndex < PlaybackBuffer.Num())
 	{
-		FModuleInputContainer& UseContainer = PlaybackBuffer[PhysicsStep - StartStep];
+		FModuleInputContainer& UseContainer = PlaybackBuffer[UseIndex];
 		InOutContainer = UseContainer;
 	}
 	else
diff --git a/Engine/Plugins/Experimental/ChaosModularVehicle/Source/ChaosModularVehicleEngine/Private/ChaosModularVehicle/ModularVehicleBaseComponent.cpp b/Engine/Plugins/Experimental/ChaosModularVehicle/Source/ChaosModularVehicleEngine/Private/ChaosModularVehicle/ModularVehicleBaseComponent.cpp
--- a/Engine/Plugins/Experimental/ChaosModularVehicle/Source/ChaosModularVehicleEngine/Private/ChaosModularVehicle/ModularVehicleBaseComponent.cpp
+++ b/Engine/Plugins/Experimental/ChaosModularVehicle/Source/ChaosModularVehicleEngine/Private/ChaosModularVehicle/ModularVehicleBaseComponent.cpp
@@ -663,6 +663,11 @@ void UModularVehicleBaseComponent::BeginPlay()
 	if (InputProducer)
 	{
 		InputProducer->InitializeContainer(CombinedInputConfiguration, InputNameMap);
+		if (InputProducer->GetTestInputBuffer())
+		{
+			// deliberate deep copy rather than copy pointer
+			VehicleSimulationPT->SetTestInputBuffer(*InputProducer->GetTestInputBuffer(), InputProducer->IsLoopingTestInputBuffer());
+		}
 	}
 
 	InputsContainer.Initialize(CombinedInputConfiguration, InputNameMap);
diff --git a/Engine/Plugins/Experimental/ChaosModularVehicle/Source/ChaosModularVehicleEngine/Private/ChaosModularVehicle/ModularVehicleSimulationCU.cpp b/Engine/Plugins/Experimental/ChaosModularVehicle/Source/ChaosModularVehicleEngine/Private/ChaosModularVehicle/ModularVehicleSimulationCU.cpp
--- a/Engine/Plugins/Experimental/ChaosModularVehicle/Source/ChaosModularVehicleEngine/Private/ChaosModularVehicle/ModularVehicleSimulationCU.cpp
+++ b/Engine/Plugins/Experimental/ChaosModularVehicle/Source/ChaosModularVehicleEngine/Private/ChaosModularVehicle/ModularVehicleSimulationCU.cpp
@@ -90,6 +90,31 @@ void FModularVehicleSimulation::SimulateModuleTree(UWorld* InWorld, float DeltaS
 		UE::TReadScopeLock InputConfigLock(InputConfigurationLock);
 
 		FModuleInputContainer Container = InputData.PhysicsInputs.NetworkInputs.VehicleInputs.Container;
+
+		if (ImplementsTestBuffer())
+		{
+			Chaos::FPBDRigidsSolver* RigidsSolver = static_cast<Chaos::FPBDRigidsSolver*>(Proxy->GetSolver<Chaos::FPhysicsSolver>());
+			check(RigidsSolver);
+			const int32 CurrentPhysicsFrame = RigidsSolver->GetCurrentFrame();
+
+			if (TestInputBufferStartFrame < 0 || TestInputBufferStartFrame > CurrentPhysicsFrame)
+			{
+				TestInputBufferStartFrame = CurrentPhysicsFrame;
+			}
+			int32 InputFrame = (CurrentPhysicsFrame - TestInputBufferStartFrame);
+
+			if (ImplementsLoopingTestBuffer() && !TestInputBuffer.IsValidIndex(InputFrame))
+			{
+				TestInputBufferStartFrame = CurrentPhysicsFrame;
+				InputFrame = 0;
+			}
+
+			if (TestInputBuffer.IsValidIndex(InputFrame))
+			{
+				Container = TestInputBuffer[InputFrame];
+			}
+		}
+
 		FInputInterface InputInterface(InputNameMap, Container);
 
 		FModuleInputContainer StateInputContainer = InputData.PhysicsInputs.StateInputs.StateInputContainer;
diff --git a/Engine/Plugins/Experimental/ChaosModularVehicle/Source/ChaosModularVehicleEngine/Public/ChaosModularVehicle/InputProducer.h b/Engine/Plugins/Experimental/ChaosModularVehicle/Source/ChaosModularVehicleEngine/Public/ChaosModularVehicle/InputProducer.h
--- a/Engine/Plugins/Experimental/ChaosModularVehicle/Source/ChaosModularVehicleEngine/Public/ChaosModularVehicle/InputProducer.h
+++ b/Engine/Plugins/Experimental/ChaosModularVehicle/Source/ChaosModularVehicleEngine/Public/ChaosModularVehicle/InputProducer.h
@@ -43,6 +43,12 @@ public:
 	/** produce input for PT simulation at PT frequency */
 	virtual void ProduceInput(int32 PhysicsStep, int32 NumSteps, const FInputNameMap& InNameMap, FModuleInputContainer& InOutContainer) override;
 
+	/** Special case override for providing test input straight onto the physics thread */
+	virtual TArray<FModuleInputContainer>* GetTestInputBuffer() { return &PlaybackBuffer; }
+
+	/** Special case override for providing test input straight onto the physics thread */
+	virtual bool IsLoopingTestInputBuffer() { return true; }
+
 	TArray<FModuleInputContainer> PlaybackBuffer;
 	int32 BufferLength = 150;
 	int32 StartStep = 0;
diff --git a/Engine/Plugins/Experimental/ChaosModularVehicle/Source/ChaosModularVehicleEngine/Public/ChaosModularVehicle/ModularVehicleSimulationCU.h b/Engine/Plugins/Experimental/ChaosModularVehicle/Source/ChaosModularVehicleEngine/Public/ChaosModularVehicle/ModularVehicleSimulationCU.h
--- a/Engine/Plugins/Experimental/ChaosModularVehicle/Source/ChaosModularVehicleEngine/Public/ChaosModularVehicle/ModularVehicleSimulationCU.h
+++ b/Engine/Plugins/Experimental/ChaosModularVehicle/Source/ChaosModularVehicleEngine/Public/ChaosModularVehicle/ModularVehicleSimulationCU.h
@@ -66,6 +66,14 @@ public:
 		StateNameMap = InNameMap;
 	}
 
+	void SetTestInputBuffer(TArray<FModuleInputContainer>& InTestInputBuffer, bool bInIsLoopBuffer)
+	{
+		UE::TWriteScopeLock InputConfigLock(InputConfigurationLock);
+		bIsLoopingTestInputBuffer = bInIsLoopBuffer;
+		TestInputBuffer = InTestInputBuffer;
+		TestInputBufferStartFrame = -1;
+	}
+
 	void CacheRootParticle(IPhysicsProxyBase* Proxy);
 
 	/** Update called from Physics Thread */
@@ -101,6 +109,12 @@ public:
 
 	int8 NetMode;
 
+	int32 TestInputBufferStartFrame = -1;
+	bool bIsLoopingTestInputBuffer = false;
+	bool ImplementsTestBuffer() { return TestInputBuffer.Num() > 0; }
+	bool ImplementsLoopingTestBuffer() { return bIsLoopingTestInputBuffer; }
+	TArray<FModuleInputContainer> TestInputBuffer;
+
 	Chaos::FPBDRigidParticleHandle* RootParticle = nullptr; // cached root particle
 
 };
diff --git a/Engine/Source/Runtime/Experimental/ChaosVehicles/ChaosVehiclesCore/Public/SimModule/ModuleInput.h b/Engine/Source/Runtime/Experimental/ChaosVehicles/ChaosVehiclesCore/Public/SimModule/ModuleInput.h
--- a/Engine/Source/Runtime/Experimental/ChaosVehicles/ChaosVehiclesCore/Public/SimModule/ModuleInput.h
+++ b/Engine/Source/Runtime/Experimental/ChaosVehicles/ChaosVehiclesCore/Public/SimModule/ModuleInput.h
@@ -732,5 +732,12 @@ public:
 
 	/** produce input for PT simulation at PT frequency */
 	virtual void ProduceInput(int32 PhysicsStep, int32 NumSteps, const FInputNameMap& InNameMap, FModuleInputContainer& InOutContainer) {}
+
+	/** Special case override for providing test input straight onto the physics thread */
+	virtual TArray<FModuleInputContainer>* GetTestInputBuffer() { return nullptr; }
+
+	/** Special case override for providing test input straight onto the physics thread */
+	virtual bool IsLoopingTestInputBuffer() { return false; }
+
 };
 
