[ChaosModularVehicles] NetToken implementation of side channel data for VehicleStates and VehicleInputs and Delta Serialization.

diff --git a/Engine/Config/BaseEngine.ini b/Engine/Config/BaseEngine.ini
--- a/Engine/Config/BaseEngine.ini
+++ b/Engine/Config/BaseEngine.ini
@@ -1486,6 +1486,8 @@ bEnsureAllClassesAreFullyPushModel=false
 +ReservedTypeIds=(StoreTypeName="StringTokenStore", TypeID=0)
 +ReservedTypeIds=(StoreTypeName="NameTokenStore", TypeID=1)
 +ReservedTypeIds=(StoreTypeName="GameplayTagTokenStore", TypeID=2)
++ReservedTypeIds=(StoreTypeName="FModuleInputNetTokenData", TypeID=3)
++ReservedTypeIds=(StoreTypeName="FNetworkModularVehicleStateNetTokenData", TypeID=4)
 
 ;Iris - end
 
diff --git a/Engine/Plugins/Experimental/ChaosModularVehicle/Source/ChaosModularVehicleEditor/ChaosModularVehicleEditor.Build.cs b/Engine/Plugins/Experimental/ChaosModularVehicle/Source/ChaosModularVehicleEditor/ChaosModularVehicleEditor.Build.cs
--- a/Engine/Plugins/Experimental/ChaosModularVehicle/Source/ChaosModularVehicleEditor/ChaosModularVehicleEditor.Build.cs
+++ b/Engine/Plugins/Experimental/ChaosModularVehicle/Source/ChaosModularVehicleEditor/ChaosModularVehicleEditor.Build.cs
@@ -23,11 +23,13 @@ namespace UnrealBuildTool.Rules
 					"BlueprintGraph",
 					"ToolMenus",
 					"AnimGraphRuntime",
-					"ChaosModularVehicleEngine"
+					"ChaosModularVehicleEngine",
+					"NetCore"
 				}
 			);
 
 			PrivateDefinitions.Add("CHAOS_INCLUDE_LEVEL_1=1");
+			SetupIrisSupport(Target);
 		}
 	}
 }
diff --git a/Engine/Plugins/Experimental/ChaosModularVehicle/Source/ChaosModularVehicleEngine/ChaosModularVehicleEngine.Build.cs b/Engine/Plugins/Experimental/ChaosModularVehicle/Source/ChaosModularVehicleEngine/ChaosModularVehicleEngine.Build.cs
--- a/Engine/Plugins/Experimental/ChaosModularVehicle/Source/ChaosModularVehicleEngine/ChaosModularVehicleEngine.Build.cs
+++ b/Engine/Plugins/Experimental/ChaosModularVehicle/Source/ChaosModularVehicleEngine/ChaosModularVehicleEngine.Build.cs
@@ -28,6 +28,7 @@ namespace UnrealBuildTool.Rules
 			);
 
 			PrivateDefinitions.Add("CHAOS_INCLUDE_LEVEL_1=1");
+			SetupIrisSupport(Target);
 		}
 	}
 }
diff --git a/Engine/Plugins/Experimental/ChaosModularVehicle/Source/ChaosModularVehicleEngine/Private/ChaosModularVehicle/ChaosSimModuleManager.cpp b/Engine/Plugins/Experimental/ChaosModularVehicle/Source/ChaosModularVehicleEngine/Private/ChaosModularVehicle/ChaosSimModuleManager.cpp
--- a/Engine/Plugins/Experimental/ChaosModularVehicle/Source/ChaosModularVehicleEngine/Private/ChaosModularVehicle/ChaosSimModuleManager.cpp
+++ b/Engine/Plugins/Experimental/ChaosModularVehicle/Source/ChaosModularVehicleEngine/Private/ChaosModularVehicle/ChaosSimModuleManager.cpp
@@ -5,8 +5,7 @@
 #include "PBDRigidsSolver.h"
 #include "GameFramework/HUD.h" // for ShowDebugInfo
 #include "Physics/Experimental/PhysScene_Chaos.h"
-#include "SimModule/ModuleFactoryRegister.h"
-
+#include "SimModule/ModuleInputTokenStore.h"
 TMap<FPhysScene*, FChaosSimModuleManager*> FChaosSimModuleManager::SceneToModuleManagerMap;
 
 FDelegateHandle FChaosSimModuleManager::OnPostWorldInitializationHandle;
@@ -100,9 +99,44 @@ void FChaosSimModuleManager::DetachFromPhysScene(FPhysScene* PhysScene)
 	FChaosSimModuleManager::SceneToModuleManagerMap.Remove(PhysScene);
 }
 
+void FChaosSimModuleManager::OnNetDriverCreated(UWorld* InWorld, UNetDriver* InNetDriver)
+{
+	if (InNetDriver)
+	{
+		if (UE::Net::FNetTokenStore* TokenStore = InNetDriver->GetNetTokenStore())
+		{
+			RegisterNetTokenDataStores(InNetDriver);
+		}
+		else
+		{
+			InNetDriver->OnNetTokenStoreReady().AddRaw(this,&FChaosSimModuleManager::RegisterNetTokenDataStores);
+		}
+	}
+}
+
+void FChaosSimModuleManager::RegisterNetTokenDataStores(UNetDriver* InNetDriver)
+{
+	if (InNetDriver)
+	{
+		if (UE::Net::FNetTokenStore* TokenStore = InNetDriver->GetNetTokenStore())
+		{
+			typedef UE::Net::TStructNetTokenDataStore<FModuleInputNetTokenData> FModuleInputNetTokenStore;
+			typedef UE::Net::TStructNetTokenDataStore<FNetworkModularVehicleStateNetTokenData> FNetworkModularVehicleStateNetTokenStore;
+			if (!TokenStore->GetDataStore<FModuleInputNetTokenStore>())
+			{
+				TokenStore->CreateAndRegisterDataStore<FModuleInputNetTokenStore>();
+			}
+			if (!TokenStore->GetDataStore<FNetworkModularVehicleStateNetTokenStore>())
+			{
+				TokenStore->CreateAndRegisterDataStore<FNetworkModularVehicleStateNetTokenStore>();
+			}
+		}
+	}
+}
 
 void FChaosSimModuleManager::RegisterCallbacks(UWorld* InWorld)
 {
+	OnNetDriverCreatedHandle = FWorldDelegates::OnNetDriverCreated.AddRaw(this, &FChaosSimModuleManager::OnNetDriverCreated);
 	OnPhysScenePreTickHandle = Scene.OnPhysScenePreTick.AddRaw(this, &FChaosSimModuleManager::Update);
 	OnPhysScenePostTickHandle = Scene.OnPhysScenePostTick.AddRaw(this, &FChaosSimModuleManager::PostUpdate);
 
@@ -120,6 +154,7 @@ void FChaosSimModuleManager::UnregisterCallbacks()
 {
 	Scene.OnPhysScenePreTick.Remove(OnPhysScenePreTickHandle);
 	Scene.OnPhysScenePostTick.Remove(OnPhysScenePostTickHandle);
+	FWorldDelegates::OnNetDriverCreated.Remove(OnNetDriverCreatedHandle);
 	
 	if (AsyncCallback)
 	{
@@ -305,5 +340,4 @@ void FChaosSimModuleManager::ParallelUpdateVehicles(float DeltaSeconds)
 		ParallelFor(AwakeVehiclesBatch.Num(), LambdaParallelUpdate2, ForceSingleThread);
 	}
 
-}
-
+}
\ No newline at end of file
diff --git a/Engine/Plugins/Experimental/ChaosModularVehicle/Source/ChaosModularVehicleEngine/Private/ChaosModularVehicle/ChaosSimModuleManagerAsyncCallback.cpp b/Engine/Plugins/Experimental/ChaosModularVehicle/Source/ChaosModularVehicleEngine/Private/ChaosModularVehicle/ChaosSimModuleManagerAsyncCallback.cpp
--- a/Engine/Plugins/Experimental/ChaosModularVehicle/Source/ChaosModularVehicleEngine/Private/ChaosModularVehicle/ChaosSimModuleManagerAsyncCallback.cpp
+++ b/Engine/Plugins/Experimental/ChaosModularVehicle/Source/ChaosModularVehicleEngine/Private/ChaosModularVehicle/ChaosSimModuleManagerAsyncCallback.cpp
@@ -5,17 +5,27 @@
 #include "ChaosModularVehicle/ModularVehicleBaseComponent.h"
 #include "ChaosModularVehicle/ModularVehicleSimulationCU.h"
 #include "PBDRigidsSolver.h"
-#include "Chaos/ParticleHandleFwd.h"
 #include "PhysicsProxy/GeometryCollectionPhysicsProxy.h"
 #include "PhysicsProxy/ClusterUnionPhysicsProxy.h"
 #include "SimModule/ModuleFactoryRegister.h"
 #include "Net/Core/NetBitArray.h"
+#include "Net/Iris/ReplicationSystem/StructNetTokenDataStoreHelper.h"
 
 FSimModuleDebugParams GSimModuleDebugParams;
 
 DECLARE_CYCLE_STAT(TEXT("AsyncCallback:OnPreSimulate_Internal"), STAT_AsyncCallback_OnPreSimulate, STATGROUP_ChaosSimModuleManager);
 DECLARE_CYCLE_STAT(TEXT("AsyncCallback:OnContactModification_Internal"), STAT_AsyncCallback_OnContactModification, STATGROUP_ChaosSimModuleManager);
 
+namespace ChaosModularVehicleCVars
+{
+	bool bEnableStateReducedBandwidth = false;
+	FAutoConsoleVariableRef EnableStateReducedBandwidth(TEXT("p.ModularVehicle.EnableStateReducedBandwidth"), bEnableStateReducedBandwidth, TEXT("Enable/Disable NetTokens and DeltaSerialization path for Modular Vehicles. Default: false"));
+	bool bEnableStateNetSerializeDebugPrinting = false;
+	FAutoConsoleVariableRef EnableStateNetSerializeDebugPrinting(TEXT("p.ModularVehicle.EnableStateNetSerializeDebugPrinting"), bEnableStateNetSerializeDebugPrinting, TEXT("Enable/Disable debug logging during NetSerialization. Default: false"));
+};
+
+UE_NET_IMPLEMENT_NAMED_NETTOKEN_STRUCT_SERIALIZERS(NetworkModularVehicleStateNetTokenData)
+
 FName FChaosSimModuleManagerAsyncCallback::GetFNameForStatId() const
 {
 	const static FLazyName StaticName("FChaosSimModuleManagerAsyncCallback");
@@ -241,10 +251,19 @@ bool FNetworkModularVehicleInputs::NetSerialize(FArchive& Ar, class UPackageMap*
 {
 	FNetworkPhysicsData::SerializeFrames(Ar);
 
-	Ar << VehicleInputs.Reverse;
-	Ar << VehicleInputs.KeepAwake;
+	Ar.SerializeBits(&VehicleInputs.Reverse, 1);
+	Ar.SerializeBits(&VehicleInputs.KeepAwake, 1);
 
-	VehicleInputs.Container.Serialize(Ar, Map, bOutSuccess);
+	FNetworkModularVehicleInputs* DeltaSource = static_cast<FNetworkModularVehicleInputs*>(DeltaSourceData);
+	if (DeltaSource && DeltaSource->ServerFrame != -1 && ChaosModularVehicleCVars::bEnableStateReducedBandwidth)
+	{
+		VehicleInputs.Container.DeltaNetSerialize(Ar, Map, bOutSuccess,DeltaSource->VehicleInputs.Container);
+	}
+	else
+	{
+		VehicleInputs.Container.Serialize(Ar, Map, bOutSuccess);
+	}
+	
 
 	bOutSuccess = true;
 	return bOutSuccess;
@@ -303,6 +322,11 @@ void FNetworkModularVehicleInputs::DecayData(float DecayAmount)
 
 bool FNetworkModularVehicleStates::NetSerialize(FArchive& Ar, class UPackageMap* Map, bool& bOutSuccess)
 {
+	if (ChaosModularVehicleCVars::bEnableStateReducedBandwidth)
+	{
+		return NetSerializeReduced(Ar, Map, bOutSuccess);
+	}
+	
 	using namespace UE::Net;
 
 	FNetworkPhysicsData::SerializeFrames(Ar);
@@ -393,6 +417,183 @@ bool FNetworkModularVehicleStates::NetSerialize(FArchive& Ar, class UPackageMap*
 	return true;
 }
 
+bool FNetworkModularVehicleStates::NetSerializeReduced(FArchive& Ar, class UPackageMap* Map, bool& bOutSuccess)
+{
+	const bool bPrintDebugInfo = ChaosModularVehicleCVars::bEnableStateNetSerializeDebugPrinting;
+	using namespace UE::Net;
+
+	FNetworkPhysicsData::SerializeFrames(Ar);
+	
+	FNetworkModularVehicleStateNetTokenData ModuleStateData;
+	ModuleStateData.Init(ModuleData);
+	TStructNetTokenDataStoreHelper<FNetworkModularVehicleStateNetTokenData>::NetSerializeAndExportToken(Ar,Map,ModuleStateData);
+
+	if (Ar.IsSaving())
+	{
+		check(ModuleStateData.Hashes.Num() == ModuleData.Num());
+	}
+	FNetworkModularVehicleStates* DeltaSource = static_cast<FNetworkModularVehicleStates*>(DeltaSourceData);
+
+	UE_LOG(LogTemp, Warning, TEXT("====NetSerialize %s. ServerFrame: %d. DeltaFrame: %d"),Ar.IsSaving()?TEXT("Save"):TEXT("Load"),ServerFrame,DeltaSource?DeltaSource->ServerFrame:-1);
+	
+	if (DeltaSource && DeltaSource->ServerFrame != -1 && DeltaSource->ModuleData.Num() == ModuleStateData.Hashes.Num() )
+	{
+		return DeltaNetSerialize(Ar,Map,bOutSuccess,ModuleStateData);
+	}
+	FBitReader* BitReader = (FBitReader*)&Ar;
+	FBitWriter* BitWriter = (FBitWriter*)&Ar;
+	UE_CLOG(bPrintDebugInfo,LogTemp, Warning, TEXT("====NetSerialize Saving: %d. ServerFrame: %d. Starting Bit: %lld"),Ar.IsSaving(),ServerFrame,Ar.IsSaving()?BitWriter->GetNumBits():BitReader->GetPosBits());
+	uint32 NumNetModules = ModuleStateData.Hashes.Num();
+
+	// Array of bits to mark which modules to serialize or not
+	if (Ar.IsLoading())
+	{
+		if (NumNetModules != ModuleData.Num())
+		{
+			ModuleData.SetNum(NumNetModules);
+		}
+		if(bPrintDebugInfo)
+		{
+			FString BitString = "";
+			for (uint32 I = 0; I < NumNetModules; I++)
+			{
+				BitString += ModuleStateData.ModuleShouldSerialize[I] ? "1" : "0";
+			}
+			UE_CLOG(bPrintDebugInfo,LogTemp, Warning, TEXT("==NetSerialize LOADING. Using ModuleShouldSerialize: %s"),*BitString);
+		}
+		for (uint32 I = 0; I < NumNetModules; I++)
+		{
+			if (NumNetModules > 0)
+			{
+				uint32 ModuleTypeHash = ModuleStateData.Hashes[I];
+				int32 SimArrayIndex = ModuleStateData.Indexes[I];
+
+				if (TSharedPtr<Chaos::FModuleNetData> Data = Chaos::FModuleFactoryRegister::Get().GenerateNetData(ModuleTypeHash, SimArrayIndex))
+				{
+					UE_CLOG(bPrintDebugInfo,LogTemp, Warning, TEXT("==NetSerialize LOADING. ModuleData: %d STA. Bit: %lld"),I,Ar.IsSaving()?BitWriter->GetNumBits():BitReader->GetPosBits());
+					ModuleData[I] = Data;
+					if (ModuleStateData.ModuleShouldSerialize[I])
+					{
+						ModuleData[I]->Serialize(Ar);
+					}
+					else
+					{
+						ModuleData[I]->ApplyDefaultState();
+					}
+					UE_CLOG(bPrintDebugInfo,LogTemp, Warning, TEXT("==NetSerialize LOADING. ModuleData: %d STA. Bit: %lld"),I,Ar.IsSaving()?BitWriter->GetNumBits():BitReader->GetPosBits());
+				}
+			}
+		}
+	}
+	else
+	{
+		if (bPrintDebugInfo)
+		{
+			FString BitString = "";
+			for (uint32 I = 0; I < NumNetModules; I++)
+			{
+				BitString += ModuleStateData.ModuleShouldSerialize[I] ? "1" : "0";
+			}
+			UE_CLOG(bPrintDebugInfo,LogTemp, Warning, TEXT("==NetSerialize SAVING. Using ModuleShouldSerialize: %s"),*BitString)
+		}
+		for (uint32 I = 0; I < NumNetModules; I++)
+		{
+			UE_CLOG(bPrintDebugInfo,LogTemp, Warning, TEXT("==NetSerialize SAVING. ModuleData: %d STA. Bit: %lld"),I,Ar.IsSaving()?BitWriter->GetNumBits():BitReader->GetPosBits())
+			if (ModuleStateData.ModuleShouldSerialize[I])
+			{
+				ModuleData[I]->Serialize(Ar);
+			}
+			UE_CLOG(bPrintDebugInfo,LogTemp, Warning, TEXT("==NetSerialize SAVING. ModuleData: %d END. Bit: %lld"),I,Ar.IsSaving()?BitWriter->GetNumBits():BitReader->GetPosBits())
+		}
+	}
+
+	bOutSuccess = true;
+	return true;
+}
+
+bool FNetworkModularVehicleStates::DeltaNetSerialize(FArchive& Ar, UPackageMap* Map, bool& bOutSuccess, FNetworkModularVehicleStateNetTokenData& VehicleStateData)
+{
+	const bool bPrintDebugInfo = ChaosModularVehicleCVars::bEnableStateNetSerializeDebugPrinting;
+	FBitReader* BitReader = (FBitReader*)&Ar;
+	FBitWriter* BitWriter = (FBitWriter*)&Ar;
+	using namespace UE::Net;
+	uint32 NumNetModules = VehicleStateData.Hashes.Num();
+	FNetworkModularVehicleStates* DeltaSource = static_cast<FNetworkModularVehicleStates*>(DeltaSourceData);
+	if (DeltaSource->ModuleData.Num() != NumNetModules)
+	{
+		return NetSerialize(Ar, Map, bOutSuccess);
+	}
+	UE_CLOG(bPrintDebugInfo,LogTemp, Warning, TEXT("====DeltaNetSerialize Saving: %d. ServerFrame: %d. DeltaSource_ServerFrame: %d Starting Bit: %lld"), Ar.IsSaving(), ServerFrame, DeltaSource->ServerFrame, Ar.IsSaving()?BitWriter->GetNumBits():BitReader->GetPosBits());
+	
+
+	TMap<FName, uint32> SerializationStash;
+	SerializationStash.Add(StashServerFrameKey, ServerFrame);
+	
+	if (Ar.IsLoading())
+	{
+		if (NumNetModules != ModuleData.Num())
+		{
+			ModuleData.SetNum(NumNetModules);
+		}
+		if(bPrintDebugInfo)
+		{
+			FString BitString = "";
+			for (uint32 I = 0; I < NumNetModules; I++)
+			{
+				BitString += VehicleStateData.ModuleShouldSerialize[I] ? "1" : "0";
+			}
+			UE_CLOG(bPrintDebugInfo,LogTemp, Warning, TEXT("==DeltaNetSerialize LOADING. Using ModuleShouldSerialize: %s"),*BitString);
+		}
+		for (uint32 I = 0; I < NumNetModules; I++)
+		{
+			int32 SimArrayIndex = VehicleStateData.Indexes[I];
+			uint32 ModuleTypeHash = VehicleStateData.Hashes[I];
+			if (TSharedPtr<Chaos::FModuleNetData> Data = Chaos::FModuleFactoryRegister::Get().GenerateNetData(ModuleTypeHash, SimArrayIndex))
+			{
+				int32 StartBit = Ar.IsSaving()?BitWriter->GetNumBits():BitReader->GetPosBits();
+				UE_CLOG(bPrintDebugInfo,LogTemp, Warning, TEXT("==DeltaNetSerialize LOADING. ModuleData: %d STA. Bit: %lld"),I,Ar.IsSaving()?BitWriter->GetNumBits():BitReader->GetPosBits());
+				ModuleData[I] = Data;
+				check(ModuleTypeHash == Chaos::FModuleFactoryRegister::GetModuleHash(Data->GetSimType()));
+				if (VehicleStateData.ModuleShouldSerialize[I])
+				{
+					ModuleData[I]->DeltaSerializeWithStash(Ar, DeltaSource->ModuleData[I].Get(),SerializationStash);
+				}
+				else
+				{
+					ModuleData[I]->ApplyDefaultState();
+				}
+				int32 EndBit = Ar.IsSaving()?BitWriter->GetNumBits():BitReader->GetPosBits();
+				UE_CLOG(bPrintDebugInfo,LogTemp, Warning, TEXT("==DeltaNetSerialize LOADING. ModuleData: %d END. Bit: %lld Total: %d Error: %d"),I,Ar.IsSaving()?BitWriter->GetNumBits():BitReader->GetPosBits(), EndBit-StartBit,Ar.IsError());
+			}
+		}
+	}
+	else
+	{
+		if(bPrintDebugInfo)
+		{
+			FString BitString = "";
+			for (uint32 I = 0; I < NumNetModules; I++)
+			{
+				BitString += VehicleStateData.ModuleShouldSerialize[I] ? "1" : "0";
+			}
+			UE_CLOG(bPrintDebugInfo,LogTemp, Warning, TEXT("==DeltaNetSerialize SAVING. Using ModuleShouldSerialize: %s"),*BitString);
+		}
+		for (uint32 I = 0; I < NumNetModules; I++)
+		{
+			int32 StartBit = Ar.IsSaving()?BitWriter->GetNumBits():BitReader->GetPosBits();
+			UE_CLOG(bPrintDebugInfo,LogTemp, Warning, TEXT("==DeltaNetSerialize SAVING. ModuleData: %d STA. Bit: %lld - %s"),I, Ar.IsSaving()?BitWriter->GetNumBits():BitReader->GetPosBits(), *DeltaSource->ModuleData[I]->GetSimType().ToString());
+			if (VehicleStateData.ModuleShouldSerialize[I])
+			{
+				ModuleData[I]->DeltaSerializeWithStash(Ar, DeltaSource->ModuleData[I].Get(), SerializationStash);
+			}
+			int32 EndBit = Ar.IsSaving()?BitWriter->GetNumBits():BitReader->GetPosBits();
+			UE_CLOG(bPrintDebugInfo,LogTemp, Warning, TEXT("==DeltaNetSerialize SAVING. ModuleData: %d END. Bit: %lld Size: %d - %s"),I, Ar.IsSaving()?BitWriter->GetNumBits():BitReader->GetPosBits(), EndBit-StartBit, *DeltaSource->ModuleData[I]->GetSimType().ToString());
+		}
+	}
+	bOutSuccess = true;
+	return true;
+}
+
 void FNetworkModularVehicleStates::ApplyData(UActorComponent* NetworkComponent) const
 {
 	if (UModularVehicleBaseComponent* ModularBaseComponent = Cast<UModularVehicleBaseComponent>(NetworkComponent))
diff --git a/Engine/Plugins/Experimental/ChaosModularVehicle/Source/ChaosModularVehicleEngine/Public/ChaosModularVehicle/ChaosSimModuleManager.h b/Engine/Plugins/Experimental/ChaosModularVehicle/Source/ChaosModularVehicleEngine/Public/ChaosModularVehicle/ChaosSimModuleManager.h
--- a/Engine/Plugins/Experimental/ChaosModularVehicle/Source/ChaosModularVehicleEngine/Public/ChaosModularVehicle/ChaosSimModuleManager.h
+++ b/Engine/Plugins/Experimental/ChaosModularVehicle/Source/ChaosModularVehicleEngine/Public/ChaosModularVehicle/ChaosSimModuleManager.h
@@ -55,6 +55,9 @@ public:
 
 	/** Post update step */
 	void PostUpdate(FChaosScene* PhysScene);
+	
+	void OnNetDriverCreated(UWorld* InWorld, UNetDriver* InNetDriver);
+	void RegisterNetTokenDataStores(UNetDriver* InNetDriver);
 
 	/** Called on GT but before Physics thread runs - PT tick rate */
 	void InjectInputs_External(int32 PhysicsStep, int32 NumSteps);
@@ -84,6 +87,7 @@ private:
 	TArray<TWeakObjectPtr<UModularVehicleBaseComponent>> CUVehicles;
 
 	// callback delegates
+	FDelegateHandle OnNetDriverCreatedHandle;
 	FDelegateHandle OnPhysScenePreTickHandle;
 	FDelegateHandle OnPhysScenePostTickHandle;
 
diff --git a/Engine/Plugins/Experimental/ChaosModularVehicle/Source/ChaosModularVehicleEngine/Public/ChaosModularVehicle/ChaosSimModuleManagerAsyncCallback.h b/Engine/Plugins/Experimental/ChaosModularVehicle/Source/ChaosModularVehicleEngine/Public/ChaosModularVehicle/ChaosSimModuleManagerAsyncCallback.h
--- a/Engine/Plugins/Experimental/ChaosModularVehicle/Source/ChaosModularVehicleEngine/Public/ChaosModularVehicle/ChaosSimModuleManagerAsyncCallback.h
+++ b/Engine/Plugins/Experimental/ChaosModularVehicle/Source/ChaosModularVehicleEngine/Public/ChaosModularVehicle/ChaosSimModuleManagerAsyncCallback.h
@@ -7,7 +7,9 @@
 #include "Chaos/SimCallbackInput.h"
 #include "Chaos/SimCallbackObject.h"
 #include "Chaos/GeometryParticlesfwd.h"
-#include "TransmissionSystem.h"
+#include "Iris/ReplicationSystem/NetTokenStructDefines.h"
+#include "Iris/ReplicationSystem/StructNetTokenDataStore.h"
+#include "Net/Iris/ReplicationSystem/StructNetTokenDataStoreHelper.h"
 #include "Physics/NetworkPhysicsComponent.h"
 #include "SimModule/SimulationModuleBase.h"
 #include "SimModule/ModuleInput.h"
@@ -65,7 +67,6 @@ struct CHAOSMODULARVEHICLEENGINE_API FModularVehicleInputs
 
 };
 
-
 /** Vehicle input data that will be used in the input history to be applied while simulating */
 USTRUCT()
 struct CHAOSMODULARVEHICLEENGINE_API FNetworkModularVehicleInputs : public FNetworkPhysicsData
@@ -104,14 +105,49 @@ struct TStructOpsTypeTraits<FNetworkModularVehicleInputs> : public TStructOpsTyp
 	};
 };
 
+USTRUCT()
+struct FNetworkModularVehicleStateNetTokenData
+{
+	GENERATED_BODY()
+	UPROPERTY()
+	TArray<uint32> Hashes;
+	UPROPERTY()
+	TArray<int32> Indexes;
+	UPROPERTY()
+	TArray<bool> ModuleShouldSerialize;
+	
+	UE_NET_NETTOKEN_GENERATED_BODY(NetworkModularVehicleStateNetTokenData, CHAOSMODULARVEHICLEENGINE_API);
+
+	CHAOSMODULARVEHICLEENGINE_API uint64 GetUniqueKey() const
+	{
+		uint64 HashOfHashes = CityHash64((const char*)Hashes.GetData(),sizeof(uint32)*Hashes.Num());
+		uint64 HashOfIndexes = CityHash64((const char*)Indexes.GetData(),sizeof(int32)*Indexes.Num());
+		uint64 HashOfShouldSerialize = CityHash64((const char*)ModuleShouldSerialize.GetData(),sizeof(bool)*ModuleShouldSerialize.Num());
+		return HashOfHashes ^ HashOfIndexes ^ HashOfShouldSerialize;
+	}
+
+	CHAOSMODULARVEHICLEENGINE_API void Init(const Chaos::FModuleNetDataArray& ModuleData)
+	{
+		for (int32 Idx = 0; Idx < ModuleData.Num(); Idx++)
+		{
+			uint32 Hash = Chaos::FModuleFactoryRegister::GetModuleHash(ModuleData[Idx]->GetSimType());
+			Hashes.Add(Hash);
+			Indexes.Add(ModuleData[Idx]->SimArrayIndex);
+			ModuleShouldSerialize.Add(!ModuleData[Idx]->IsDefaultState());
+		}
+	}
+};
+
+UE_NET_DECLARE_NAMED_NETTOKEN_STRUCT_SERIALIZERS(NetworkModularVehicleStateNetTokenData, CHAOSMODULARVEHICLEENGINE_API);
+
 /** Vehicle state data that will be used in the state history to rewind the simulation */
 USTRUCT()
 struct CHAOSMODULARVEHICLEENGINE_API FNetworkModularVehicleStates : public FNetworkPhysicsData
 {
-	GENERATED_USTRUCT_BODY()
+	GENERATED_BODY()
 
+	inline static FName StashServerFrameKey = FName("ServerFrame");
 	Chaos::FModuleNetDataArray ModuleData;
-
 	/**  Apply the data onto the network physics component */
 	virtual void ApplyData(UActorComponent* NetworkComponent) const override;
 
@@ -120,6 +156,8 @@ struct CHAOSMODULARVEHICLEENGINE_API FNetworkModularVehicleStates : public FNetw
 
 	/**  Serialize data function that will be used to transfer the struct across the network */
 	bool NetSerialize(FArchive& Ar, class UPackageMap* Map, bool& bOutSuccess);
+	bool NetSerializeReduced(FArchive& Ar, class UPackageMap* Map, bool& bOutSuccess);
+	bool DeltaNetSerialize(FArchive& Ar, class UPackageMap* Map, bool& BOutSuccess, FNetworkModularVehicleStateNetTokenData& VehicleStateData);
 
 	/** Interpolate the data in between two states */
 	virtual void InterpolateData(const FNetworkPhysicsData& MinData, const FNetworkPhysicsData& MaxData) override;
diff --git a/Engine/Source/Runtime/Experimental/ChaosVehicles/ChaosVehiclesCore/ChaosVehiclesCore.Build.cs b/Engine/Source/Runtime/Experimental/ChaosVehicles/ChaosVehiclesCore/ChaosVehiclesCore.Build.cs
--- a/Engine/Source/Runtime/Experimental/ChaosVehicles/ChaosVehiclesCore/ChaosVehiclesCore.Build.cs
+++ b/Engine/Source/Runtime/Experimental/ChaosVehicles/ChaosVehiclesCore/ChaosVehiclesCore.Build.cs
@@ -10,11 +10,14 @@ namespace UnrealBuildTool.Rules
                 new string[] {
                 "Core",
                 "CoreUObject",
-				"Chaos"
+				"Chaos",
+				"NetCore",
+				"Engine"
 				}
             ); 
 
 			PrivateDefinitions.Add("CHAOS_INCLUDE_LEVEL_1=1");
+			SetupIrisSupport(Target);
 		}
     }
 }
diff --git a/Engine/Source/Runtime/Experimental/ChaosVehicles/ChaosVehiclesCore/Private/SimModule/ModuleInput.cpp b/Engine/Source/Runtime/Experimental/ChaosVehicles/ChaosVehiclesCore/Private/SimModule/ModuleInput.cpp
--- a/Engine/Source/Runtime/Experimental/ChaosVehicles/ChaosVehiclesCore/Private/SimModule/ModuleInput.cpp
+++ b/Engine/Source/Runtime/Experimental/ChaosVehicles/ChaosVehiclesCore/Private/SimModule/ModuleInput.cpp
@@ -1,12 +1,22 @@
 // Copyright Epic Games, Inc. All Rights Reserved.
 
 #include "SimModule/ModuleInput.h"
+#include "SimModule/ModuleInputTokenStore.h"
+#include "Net/Iris/ReplicationSystem/StructNetTokenDataStoreHelper.h"
 #include "Net/Core/NetBitArray.h"
 
 TArray<struct FModuleInputSetup>* FScopedModuleInputInitializer::InitSetupData = nullptr;
 
 DEFINE_LOG_CATEGORY(LogModularInput);
 
+namespace ChaosModularVehicleCVars
+{
+	bool bEnableInputReducedBandwidth = false;
+	FAutoConsoleVariableRef EnableInputReducedBandwidth(TEXT("p.ModularVehicle.EnableInputReducedBandwidth"), bEnableInputReducedBandwidth, TEXT("Enable/Disable NetTokens and DeltaSerialization path for Modular Vehicle Input. Default: false"));
+};
+
+UE_NET_IMPLEMENT_NAMED_NETTOKEN_STRUCT_SERIALIZERS(ModuleInputNetTokenData)
+
 bool FModuleInputValue::IsNonZero(float Tolerance) const
 {
 	switch (ValueType)
@@ -112,12 +122,13 @@ void FModuleInputValue::SetMagnitude(float NewSize)
 
 void FModuleInputValue::Serialize(FArchive& Ar, UPackageMap* Map, bool& bOutSuccess)
 {
-	// #TODO: send only value changes/deltas
-
 	bool bIsNonZero = Ar.IsLoading() ? false : IsQuantizedNonZero();
 	Ar.SerializeBits(&bIsNonZero, 1);
-	Ar.SerializeBits(&ValueType, 3);
-	Ar.SerializeBits(&bApplyInputDecay, 1);
+	if (!ChaosModularVehicleCVars::bEnableInputReducedBandwidth)
+	{
+		Ar.SerializeBits(&ValueType, 3);
+		Ar.SerializeBits(&bApplyInputDecay, 1);
+	}
 
 	if (bIsNonZero)
 	{
@@ -160,6 +171,31 @@ void FModuleInputValue::Serialize(FArchive& Ar, UPackageMap* Map, bool& bOutSucc
 	bOutSuccess = true;
 }
 
+void FModuleInputValue::DeltaNetSerialize(FArchive& Ar, UPackageMap* Map, bool& bOutSuccess, const FModuleInputValue& PreviousInputValue)
+{
+	bool bHasSamePreviousValue = Value == PreviousInputValue.Value;
+	bool bHasSamePreviousType = ValueType == PreviousInputValue.ValueType;
+	bool bHasSamePreviousValueInt = ValueInt == PreviousInputValue.ValueInt;
+	bool bIsSameAsPrevious = bHasSamePreviousValue && bHasSamePreviousType && bHasSamePreviousValueInt;
+	Ar.SerializeBits(&bIsSameAsPrevious, 1);
+	if (bIsSameAsPrevious)
+	{
+		bOutSuccess = true;
+		if (Ar.IsSaving())
+		{
+			return;
+		}
+
+		ConvertToType(PreviousInputValue.ValueType);
+		Value = PreviousInputValue.Value;
+		ValueInt = PreviousInputValue.ValueInt;
+		bApplyInputDecay = PreviousInputValue.bApplyInputDecay;
+		return;
+	}
+
+	Serialize(Ar, Map, bOutSuccess);
+}
+
 void FModuleInputValue::Lerp(const FModuleInputValue& Min, const FModuleInputValue& Max, float Alpha)
 {
 	switch (GetValueType())
@@ -332,6 +368,10 @@ void FModuleInputContainer::ZeroValues()
 
 void FModuleInputContainer::Serialize(FArchive& Ar, UPackageMap* Map, bool& bOutSuccess)
 {
+	if (ChaosModularVehicleCVars::bEnableInputReducedBandwidth)
+	{
+		return SerializeReduced(Ar, Map, bOutSuccess);
+	}
 	using namespace UE::Net;
 
 	bOutSuccess = true;
@@ -352,6 +392,65 @@ void FModuleInputContainer::Serialize(FArchive& Ar, UPackageMap* Map, bool& bOut
 	return;
 }
 
+void FModuleInputContainer::SerializeReduced(FArchive& Ar, UPackageMap* Map, bool& bOutSuccess)
+{
+	using namespace UE::Net;
+    bOutSuccess = true;
+    
+    FModuleInputNetTokenData InputStateData;
+    InputStateData.Init(InputValues);
+    TStructNetTokenDataStoreHelper<FModuleInputNetTokenData>::NetSerializeAndExportToken(Ar,Map,InputStateData);
+
+    uint32 Number = InputStateData.Types.Num();
+    //Ar.SerializeIntPacked(Number);
+
+    if (Ar.IsLoading())
+    {
+    	InputValues.SetNum(Number);
+    }
+
+    for (uint32 I = 0; I < Number; I++)
+    {
+    	InputValues[I].ConvertToType(static_cast<EModuleInputValueType>(InputStateData.Types[I]));
+    	InputValues[I].SetApplyInputDecay(InputStateData.DecayValues[I]);
+    	InputValues[I].Serialize(Ar, Map, bOutSuccess);
+    }
+
+    return;
+}
+
+void FModuleInputContainer::DeltaNetSerialize(FArchive& Ar, UPackageMap* Map, bool& bOutSuccess, FModuleInputContainer& PreviousInputContainer)
+{
+	using namespace UE::Net;
+	bOutSuccess = true;
+	
+	FModuleInputNetTokenData InputStateData;
+	InputStateData.Init(InputValues);
+	TStructNetTokenDataStoreHelper<FModuleInputNetTokenData>::NetSerializeAndExportToken(Ar,Map,InputStateData);
+
+	uint32 Number = InputStateData.Types.Num();
+	if (Ar.IsLoading())
+	{
+		InputValues.SetNum(Number);
+	}
+
+	for (uint32 I = 0; I < Number; I++)
+	{
+		InputValues[I].ConvertToType(static_cast<EModuleInputValueType>(InputStateData.Types[I]));
+		InputValues[I].SetApplyInputDecay(InputStateData.DecayValues[I]);
+		if (PreviousInputContainer.InputValues.Num() == InputValues.Num())
+		{
+			InputValues[I].DeltaNetSerialize(Ar, Map, bOutSuccess, PreviousInputContainer.InputValues[I]);
+		}
+		else
+		{
+			InputValues[I].Serialize(Ar, Map, bOutSuccess);
+		}
+		
+	}
+
+}
+
 int FModuleInputContainer::AddInput(EModuleInputValueType Type, TSubclassOf<UDefaultModularVehicleInputModifier>& InputModifierClass)
 {
 	switch(Type)
diff --git a/Engine/Source/Runtime/Experimental/ChaosVehicles/ChaosVehiclesCore/Private/SimModule/SimModuleTree.cpp b/Engine/Source/Runtime/Experimental/ChaosVehicles/ChaosVehiclesCore/Private/SimModule/SimModuleTree.cpp
--- a/Engine/Source/Runtime/Experimental/ChaosVehicles/ChaosVehiclesCore/Private/SimModule/SimModuleTree.cpp
+++ b/Engine/Source/Runtime/Experimental/ChaosVehicles/ChaosVehiclesCore/Private/SimModule/SimModuleTree.cpp
@@ -533,7 +533,7 @@ void FSimModuleTree::GenerateReplicationStructure(Chaos::FModuleNetDataArray& Ne
 	}
 
 	const TArray<FSimModuleNode>& Tree = SimulationModuleTree;
-	NetData.Reserve(Tree.Num());
+	NetData.Reset(Tree.Num());
 	for (int Index = 0; Index < Tree.Num(); Index++)
 	{
 		if (ISimulationModuleBase* SimModule = Tree[Index].SimModule)
@@ -557,7 +557,7 @@ void FSimModuleTree::SetNetState(Chaos::FModuleNetDataArray& ModuleDatas)
 		return;
 	}
 
-	if (ModuleDatas.IsEmpty())
+	//Always Regenerate Replication Structure to use unique data per ReplicationStructure
 	{
 		GenerateReplicationStructure(ModuleDatas);
 	}
diff --git a/Engine/Source/Runtime/Experimental/ChaosVehicles/ChaosVehiclesCore/Public/SimModule/ModuleInput.h b/Engine/Source/Runtime/Experimental/ChaosVehicles/ChaosVehiclesCore/Public/SimModule/ModuleInput.h
--- a/Engine/Source/Runtime/Experimental/ChaosVehicles/ChaosVehiclesCore/Public/SimModule/ModuleInput.h
+++ b/Engine/Source/Runtime/Experimental/ChaosVehicles/ChaosVehiclesCore/Public/SimModule/ModuleInput.h
@@ -3,7 +3,6 @@
 #pragma once
 
 #include "CoreMinimal.h"
-//#include "Engine/NetSerialization.h"
 #include "Templates/SubclassOf.h"
 #include "ModuleInput.generated.h"
 
@@ -103,43 +102,65 @@ namespace ModularQuantize
 	}
 
 	template<int32 MaxValue, uint32 NumBits, typename T UE_REQUIRES(std::is_floating_point_v<T>&& NumBits < 32)>
-	bool WriteCompressedFloat(const T Value, FArchive& Ar)
+	bool WriteCompressedFloat(const T Value, FArchive& Ar, bool bAutoReduceZeroValues = false)
 	{
 		using Details = ModularQuantize::TCompressedFloatDetails<MaxValue, NumBits>;
 
 		uint32 CompressedValue;
 		bool clamp = ModularQuantize::ToCompressedFloat<MaxValue, NumBits>(Value, CompressedValue);
-
-		Ar.SerializeInt(CompressedValue, Details::SerIntMax);
+		
+		if(bAutoReduceZeroValues)
+		{
+			bool IsZero = CompressedValue == 0;
+			Ar.SerializeBits(&IsZero,1);
+			if(!IsZero)
+			{
+				Ar.SerializeBits(&CompressedValue, NumBits);
+			}
+		}
+		else
+		{
+			Ar.SerializeBits(&CompressedValue, NumBits);
+		}
 
 		return !clamp;
 	}
 
 	template<int32 MaxValue, uint32 NumBits, typename T UE_REQUIRES(std::is_floating_point_v<T>&& NumBits < 32)>
-	bool ReadCompressedFloat(T& Value, FArchive& Ar)
+	bool ReadCompressedFloat(T& Value, FArchive& Ar, bool bAutoReduceZeroValues = false)
 	{
 		using Details = ModularQuantize::TCompressedFloatDetails<MaxValue, NumBits>;
-
-		uint32 CompressedValue;
-		Ar.SerializeInt(CompressedValue, Details::SerIntMax);
+		uint32 CompressedValue = 0;
+		if(bAutoReduceZeroValues)
+		{
+			bool IsZero;
+			Ar.SerializeBits(&IsZero,1);
+			if(!IsZero)
+			{
+				Ar.SerializeBits(&CompressedValue, NumBits);
+			}
+		}
+		else
+		{
+			Ar.SerializeBits(&CompressedValue,NumBits);
+		}
 
 		ModularQuantize::FromCompressedFloat<MaxValue, NumBits>(CompressedValue, Value);
-
 		return true;
 	}
 
 	// Required because we seialize quantized vector in seperate parts depending on input type
 	template<int32 MaxValue, uint32 NumBits>
-	bool SerializeFixedFloat(double& InOutValue, FArchive& Ar)
+	bool SerializeFixedFloat(double& InOutValue, FArchive& Ar, bool bAutoReduceZeroValues = false)
 	{
 		if (Ar.IsSaving())
 		{
 			bool success = true;
-			success &= ModularQuantize::WriteCompressedFloat<MaxValue, NumBits>(InOutValue, Ar);
+			success &= ModularQuantize::WriteCompressedFloat<MaxValue, NumBits>(InOutValue, Ar, bAutoReduceZeroValues);
 			return success;
 		}
 
-		ModularQuantize::ReadCompressedFloat<MaxValue, NumBits>(InOutValue, Ar);
+		ModularQuantize::ReadCompressedFloat<MaxValue, NumBits>(InOutValue, Ar, bAutoReduceZeroValues);
 		return true;
 	}
 
@@ -150,6 +171,18 @@ namespace ModularQuantize
 		ModularQuantize::ToCompressedFloat<MaxValue, NumBits>(Value, CompressedValue);
 		ModularQuantize::FromCompressedFloat<MaxValue, NumBits>(CompressedValue, Value);
 	}
+	
+	template<int32 MaxValue, uint32 NumBits, typename T UE_REQUIRES(std::is_floating_point_v<T>&& NumBits < 32)>
+    bool QuantizedIsNearlyEqual(const T& Left,const T& Right)
+    {
+    	uint32 CompressedValue = 0;
+    	T QuantLeft, QuantRight;
+    	ModularQuantize::ToCompressedFloat<MaxValue, NumBits>(Left, CompressedValue);
+    	ModularQuantize::FromCompressedFloat<MaxValue, NumBits>(CompressedValue, QuantLeft);
+    	ModularQuantize::ToCompressedFloat<MaxValue, NumBits>(Right, CompressedValue);
+    	ModularQuantize::FromCompressedFloat<MaxValue, NumBits>(CompressedValue, QuantRight);
+    	return FMath::IsNearlyEqual(QuantLeft, QuantRight);
+    }
 
 }
 
@@ -326,8 +359,6 @@ public:
 		return OutValue;
 	}
 
-
-
 	template<typename T>
 	inline T Get() const { static_assert(sizeof(T) == 0, "Unsupported conversion for type"); }
 
@@ -356,6 +387,7 @@ public:
 
 	// Serialize values
 	void Serialize(FArchive& Ar, UPackageMap* Map, bool& bOutSuccess);
+	void DeltaNetSerialize(FArchive& Ar, UPackageMap* Map, bool& bOutSuccess,const FModuleInputValue& PreviousInputValue);
 
 	void Lerp(const FModuleInputValue& Min, const FModuleInputValue& Max, float Alpha);
 
@@ -368,7 +400,7 @@ public:
 
 	/** During physics resimulation, apply decay for this input while it's being extrapolated */
 	void SetApplyInputDecay(const bool bInApplyInputDecay) { bApplyInputDecay = bInApplyInputDecay; }
-	bool ShouldApplyInputDecay() { return  bApplyInputDecay; }
+	bool ShouldApplyInputDecay() const { return  bApplyInputDecay; }
 
 protected:
 	void SetMagnitude(float NewSize);
@@ -563,7 +595,6 @@ private:
 	static TArray<struct FModuleInputSetup>* InitSetupData;
 };
 
-
 USTRUCT()
 struct CHAOSVEHICLESCORE_API FModuleInputContainer
 {
@@ -616,6 +647,8 @@ public:
 	void ZeroValues();
 
 	void Serialize(FArchive& Ar, UPackageMap* Map, bool& bOutSuccess);
+	void SerializeReduced(FArchive& Ar, UPackageMap* Map, bool& bOutSuccess);
+	void DeltaNetSerialize(FArchive& Ar, UPackageMap* Map, bool& bOutSuccess, FModuleInputContainer& PreviousInputContainer);
 
 	int AddInput(EModuleInputValueType Type, TSubclassOf<UDefaultModularVehicleInputModifier>& InputModifierClass);
 
@@ -634,11 +667,9 @@ private:
 
 };
 
-
 class CHAOSVEHICLESCORE_API FInputInterface
 {
 public:
-
 	using FInputNameMap = TMap<FName, int>;
 
 	FInputInterface(const FInputNameMap& InNameMap, FModuleInputContainer& InValueContainer)
diff --git a/Engine/Source/Runtime/Experimental/ChaosVehicles/ChaosVehiclesCore/Public/SimModule/ModuleInputTokenStore.h b/Engine/Source/Runtime/Experimental/ChaosVehicles/ChaosVehiclesCore/Public/SimModule/ModuleInputTokenStore.h
new file mode 100644
--- /dev/null
+++ b/Engine/Source/Runtime/Experimental/ChaosVehicles/ChaosVehiclesCore/Public/SimModule/ModuleInputTokenStore.h
@@ -0,0 +1,39 @@
+// Copyright Epic Games, Inc. All Rights Reserved.
+
+#pragma once
+
+#include "Iris/ReplicationSystem/NetTokenStructDefines.h"
+#include "Iris/ReplicationSystem/StructNetTokenDataStore.h"
+#include "Net/Iris/ReplicationSystem/StructNetTokenDataStoreHelper.h"
+#include "ModuleInput.h"
+#include "ModuleInputTokenStore.generated.h"
+
+USTRUCT()
+struct FModuleInputNetTokenData
+{
+	GENERATED_BODY()
+	UPROPERTY()
+	TArray<uint8> Types;
+	UPROPERTY()
+	TArray<bool> DecayValues;
+	
+	UE_NET_NETTOKEN_GENERATED_BODY(ModuleInputNetTokenData, CHAOSVEHICLESCORE_API)
+	
+	CHAOSVEHICLESCORE_API uint64 GetUniqueKey() const
+	{
+		uint64 HashOfTypes = CityHash64((const char*)Types.GetData(),sizeof(EModuleInputValueType)*Types.Num());
+		uint64 HashOfDecayValues = CityHash64((const char*)DecayValues.GetData(),sizeof(bool)*DecayValues.Num());
+		return HashOfTypes ^ HashOfDecayValues;
+	}
+
+	CHAOSVEHICLESCORE_API void Init(const TArray<FModuleInputValue>& ModuleInputs)
+	{
+		for (int32 Idx = 0; Idx < ModuleInputs.Num(); Idx++)
+		{
+			Types.Add(static_cast<uint8>(ModuleInputs[Idx].GetValueType()));
+			DecayValues.Add(ModuleInputs[Idx].ShouldApplyInputDecay());
+		}
+	}
+};
+
+UE_NET_DECLARE_NAMED_NETTOKEN_STRUCT_SERIALIZERS(ModuleInputNetTokenData, CHAOSVEHICLESCORE_API);
\ No newline at end of file
diff --git a/Engine/Source/Runtime/Experimental/ChaosVehicles/ChaosVehiclesCore/Public/SimModule/SimulationModuleBase.h b/Engine/Source/Runtime/Experimental/ChaosVehicles/ChaosVehiclesCore/Public/SimModule/SimulationModuleBase.h
--- a/Engine/Source/Runtime/Experimental/ChaosVehicles/ChaosVehiclesCore/Public/SimModule/SimulationModuleBase.h
+++ b/Engine/Source/Runtime/Experimental/ChaosVehicles/ChaosVehiclesCore/Public/SimModule/SimulationModuleBase.h
@@ -227,6 +227,18 @@ namespace Chaos
 		{
 			return MostRecentAdd;
 		}
+		template <typename U>
+		static FName GetStaticSimType()
+		{
+			if constexpr (TSimModuleTypeIsRecursive_v<U>)
+			{
+				return U::RecurseSimType(); 
+			}
+			else
+			{
+				return TSimulationModuleTypeable<U>::StaticSimType(); 
+			}
+		}
 		template<typename U>
 		bool IsSimType() const
 		{
@@ -511,6 +523,8 @@ namespace Chaos
 		virtual ~FModuleNetData() {}
 
 		virtual void Serialize(FArchive& Ar) = 0;
+		virtual void DeltaSerialize(FArchive& Ar,const FModuleNetData* Previous) {Serialize(Ar);}
+		virtual void DeltaSerializeWithStash(FArchive& Ar,const FModuleNetData* Previous, TMap<FName,uint32>& InOutSerializationStash) {DeltaSerialize(Ar,Previous);}
 		virtual void FillNetState(const ISimulationModuleBase* SimModule) = 0;
 		virtual void FillSimState(ISimulationModuleBase* SimModule) = 0;
 		virtual void Lerp(const float LerpFactor, const FModuleNetData& Max, const FModuleNetData& MaxValue) = 0;
@@ -586,7 +600,7 @@ namespace Chaos
 			return MakeShared<T>(
 				SimArrayIndex
 	#if !(UE_BUILD_SHIPPING || UE_BUILD_TEST)
-				, DebugString
+				, FSimulationModuleTypeContainer::GetStaticSimType<T>().ToString()
 	#endif			
 			);
 		}
