Engine Part of Modular Vehicle Scene Graph code
All marked experimental and internal, not included in any projects in this CL

#okforversepublic

diff --git a/Engine/Plugins/Experimental/ChaosModularVehicle/Source/ChaosModularVehicleEngine/Private/ChaosModularVehicle/ChaosSimModuleManager.cpp b/Engine/Plugins/Experimental/ChaosModularVehicle/Source/ChaosModularVehicleEngine/Private/ChaosModularVehicle/ChaosSimModuleManager.cpp
--- a/Engine/Plugins/Experimental/ChaosModularVehicle/Source/ChaosModularVehicleEngine/Private/ChaosModularVehicle/ChaosSimModuleManager.cpp
+++ b/Engine/Plugins/Experimental/ChaosModularVehicle/Source/ChaosModularVehicleEngine/Private/ChaosModularVehicle/ChaosSimModuleManager.cpp
@@ -185,16 +185,20 @@ void FChaosSimModuleManager::AddVehicle(TWeakObjectPtr<UModularVehicleBaseCompon
 
 void FChaosSimModuleManager::RemoveVehicle(TWeakObjectPtr<UModularVehicleBaseComponent> Vehicle)
 {
-	check(Vehicle != NULL);
-
-	CUVehicles.Remove(Vehicle);
+	if (Vehicle != nullptr)
+	{
+		CUVehicles.Remove(Vehicle);
+	}
 }
 
 void FChaosSimModuleManager::ScenePreTick(FPhysScene* PhysScene, float DeltaTime)
 {
 	for (int32 i = 0; i < CUVehicles.Num(); ++i)
 	{
-		CUVehicles[i]->PreTickGT(DeltaTime);
+		if (CUVehicles[i] != nullptr)
+		{
+			CUVehicles[i]->PreTickGT(DeltaTime);
+		}
 	}
 }
 
diff --git a/Engine/Plugins/Experimental/ChaosModularVehicle/Source/ChaosModularVehicleEngine/Private/ChaosModularVehicle/ChaosSimModuleManagerAsyncCallback.cpp b/Engine/Plugins/Experimental/ChaosModularVehicle/Source/ChaosModularVehicleEngine/Private/ChaosModularVehicle/ChaosSimModuleManagerAsyncCallback.cpp
--- a/Engine/Plugins/Experimental/ChaosModularVehicle/Source/ChaosModularVehicleEngine/Private/ChaosModularVehicle/ChaosSimModuleManagerAsyncCallback.cpp
+++ b/Engine/Plugins/Experimental/ChaosModularVehicle/Source/ChaosModularVehicleEngine/Private/ChaosModularVehicle/ChaosSimModuleManagerAsyncCallback.cpp
@@ -222,6 +222,11 @@ void FModularVehicleAsyncInput::ProcessInputs()
 		return;
 	}
 
+	if (!GetVehicle()->VehicleSimulationPT)
+	{
+		return;
+	}
+
 	FModularVehicleSimulation* VehicleSim = GetVehicle()->VehicleSimulationPT.Get();
 
 	if (VehicleSim == nullptr || !GetVehicle()->bUsingNetworkPhysicsPrediction || GetVehicle()->GetWorld() == nullptr)
diff --git a/Engine/Plugins/Experimental/ChaosModularVehicle/Source/ChaosModularVehicleEngine/Private/ChaosModularVehicle/ModularVehicleBaseComponent.cpp b/Engine/Plugins/Experimental/ChaosModularVehicle/Source/ChaosModularVehicleEngine/Private/ChaosModularVehicle/ModularVehicleBaseComponent.cpp
--- a/Engine/Plugins/Experimental/ChaosModularVehicle/Source/ChaosModularVehicleEngine/Private/ChaosModularVehicle/ModularVehicleBaseComponent.cpp
+++ b/Engine/Plugins/Experimental/ChaosModularVehicle/Source/ChaosModularVehicleEngine/Private/ChaosModularVehicle/ModularVehicleBaseComponent.cpp
@@ -553,13 +553,13 @@ void UModularVehicleBaseComponent::CreateIndependentSimComponents(USceneComponen
 
 int UModularVehicleBaseComponent::AddSimModule(Chaos::ISimulationModuleBase* NewModule, const FTransform& LocalTransform, int ParentIndex, int TransformIndex)
 {
-	TUniquePtr<Chaos::FSimModuleTree>& SimModuleTree = VehicleSimulationPT->AccessSimComponentTree();
+	if (VehicleSimulationPT == nullptr)
+	{
+		ensureMsgf(false, TEXT("VehicleSimulationPT ptr is null, This is unexpected behaviour."));
+		return INDEX_NONE;
+	}
 
-#if !(UE_BUILD_SHIPPING || UE_BUILD_TEST)
-	FString DebugString;
-	NewModule->GetDebugString(DebugString);
-	UE_LOG(LogModularBase, Log, TEXT("CreateAssociatedSimComponents Found Sim Component %s (%s)"), *NewModule->GetDebugName(), *DebugString);
-#endif
+	TUniquePtr<Chaos::FSimModuleTree>& SimModuleTree = VehicleSimulationPT->AccessSimComponentTree();
 
 	FVector LocalOffset(0, 0, 0);
 	FVector LocalOffsetCOM(0, 0, 0);
@@ -585,6 +585,12 @@ int UModularVehicleBaseComponent::AddSimModule(Chaos::ISimulationModuleBase* New
 	FTransform ClusterredTransform(FQuat::Identity, InitialTransform.GetLocation());
 	NewModule->SetClusteredTransform(ClusterredTransform);
 
+	NewModule->SetAnimationData(NewModule->GetBoneName(), NewModule->GetAnimationOffset(), ModuleAnimationSetups.Num());
+	FModuleAnimationSetup AnimSetup(NewModule->GetBoneName(), TransformIndex, NewModule->GetGuid());
+	AnimSetup.InitialRotOffset = InitialTransform.GetRotation();
+	AnimSetup.InitialLocOffset = InitialTransform.GetTranslation();
+	ModuleAnimationSetups.Add(AnimSetup);
+
 	return TreeIndex;
 }
 
@@ -600,6 +606,7 @@ void UModularVehicleBaseComponent::FinalizeModuleUpdates()
 
 	if (RootPhysicsObject == nullptr)
 	{
+		ensure(GetPhysicsProxy() != nullptr); // should never happen
 		CacheRootPhysicsObject(GetPhysicsProxy()); //#TODO where to put / refresh?
 	}
 
@@ -679,8 +686,11 @@ void UModularVehicleBaseComponent::BeginPlay()
 	}
 	// #TODO reinstate ? GenerateInputModifiers(CombinedInputConfiguration);
 
-	VehicleSimulationPT->SetInputMappings(InputNameMap);
-	VehicleSimulationPT->SetStateMappings(StateNameMap);
+	if (VehicleSimulationPT)
+	{
+		VehicleSimulationPT->SetInputMappings(InputNameMap);
+		VehicleSimulationPT->SetStateMappings(StateNameMap);
+	}
 
 }
 
@@ -1670,4 +1680,4 @@ void UModularVehicleBaseComponent::LogInputSetup()
 		UE_LOG(LogModularBase, Log, TEXT("Input: %s %d"), *NamePair.Key.ToString(), NamePair.Value);
 	}
 #endif //!(UE_BUILD_SHIPPING || UE_BUILD_TEST)
-}
\ No newline at end of file
+}
diff --git a/Engine/Plugins/Experimental/ChaosModularVehicle/Source/ChaosModularVehicleEngine/Public/ChaosModularVehicle/ModularVehicleBaseComponent.h b/Engine/Plugins/Experimental/ChaosModularVehicle/Source/ChaosModularVehicleEngine/Public/ChaosModularVehicle/ModularVehicleBaseComponent.h
--- a/Engine/Plugins/Experimental/ChaosModularVehicle/Source/ChaosModularVehicleEngine/Public/ChaosModularVehicle/ModularVehicleBaseComponent.h
+++ b/Engine/Plugins/Experimental/ChaosModularVehicle/Source/ChaosModularVehicleEngine/Public/ChaosModularVehicle/ModularVehicleBaseComponent.h
@@ -64,13 +64,16 @@ struct CHAOSMODULARVEHICLEENGINE_API FModuleAnimationSetup
 {
 	GENERATED_USTRUCT_BODY()
 
-	FModuleAnimationSetup(FName BoneNameIn, int TransformIndexIn)
+	FModuleAnimationSetup(FName BoneNameIn, int TransformIndexIn, int GuidIn = INDEX_NONE)
 		: BoneName(BoneNameIn)
 		, RotOffset(FRotator::ZeroRotator)
 		, LocOffset(FVector::ZeroVector)
 		, CombinedRotation(FQuat::Identity)
 		, AnimFlags(0)
 		, TransformIndex(TransformIndexIn)
+		, ModuleGUID(GuidIn)
+		, InitialRotOffset(FQuat::Identity)
+		, InitialLocOffset(FVector::ZeroVector)
 	{
 	}
 
@@ -81,6 +84,9 @@ struct CHAOSMODULARVEHICLEENGINE_API FModuleAnimationSetup
 		, CombinedRotation(FQuat::Identity)
 		, AnimFlags(0)
 		, TransformIndex(INDEX_NONE)
+		, ModuleGUID(INDEX_NONE)
+		, InitialRotOffset(FQuat::Identity)
+		, InitialLocOffset(FVector::ZeroVector)
 	{
 	}
 
@@ -90,6 +96,10 @@ struct CHAOSMODULARVEHICLEENGINE_API FModuleAnimationSetup
 	FQuat CombinedRotation;
 	uint16 AnimFlags;
 	int TransformIndex; // required for non skeletal mesh case
+	int ModuleGUID;
+	FQuat InitialRotOffset;
+	FVector InitialLocOffset;
+
 };
 
 
@@ -371,7 +381,6 @@ public:
 	void RemoveSimModule(int ModuleGuid);
 	void FinalizeModuleUpdates();
 
-	void EnableModuleCreationFromActorComponents(bool bEnable) { bEnableModuleCreationFromActorComponents = bEnable; }
 	Chaos::FSimTreeUpdates StoredTreeUpdates;
 
 	TArray<AActor*> ActorsToIgnore;
@@ -413,7 +422,6 @@ private:
 	int32 ClusteringCount = 0;
 
 	bool bIsLocallyControlled;
-	bool bEnableModuleCreationFromActorComponents = true;
 
 	TArray<FModuleAnimationSetup> ModuleAnimationSetups;
 
diff --git a/Engine/Plugins/Experimental/ChaosModularVehicle/Source/ChaosModularVehicleEngine/Public/ChaosModularVehicle/VehicleSimBaseComponent.h b/Engine/Plugins/Experimental/ChaosModularVehicle/Source/ChaosModularVehicleEngine/Public/ChaosModularVehicle/VehicleSimBaseComponent.h
--- a/Engine/Plugins/Experimental/ChaosModularVehicle/Source/ChaosModularVehicleEngine/Public/ChaosModularVehicle/VehicleSimBaseComponent.h
+++ b/Engine/Plugins/Experimental/ChaosModularVehicle/Source/ChaosModularVehicleEngine/Public/ChaosModularVehicle/VehicleSimBaseComponent.h
@@ -54,6 +54,7 @@ public:
 	virtual Chaos::ISimulationModuleBase* CreateNewCoreModule() const { return nullptr; }
 	virtual FName GetBoneName() const { return NAME_None; }
 	virtual const FVector& GetAnimationOffset() const { return FVector::ZeroVector; }
+	virtual void SetAnimationEnabled(bool AnimationEnabledIn) { }
 	virtual bool GetAnimationEnabled() const { return false; }
 	virtual TArray<FModuleInputSetup> GetInputConfig() const { return TArray<FModuleInputSetup>(); }
 	virtual int32 GetAnimationSetupIndex() const { return INDEX_NONE; }
@@ -91,6 +92,7 @@ public:
 	/** IVehicleSimBaseComponentInterface overrides */
 	virtual FName GetBoneName() const override { return BoneName; }
 	virtual const FVector& GetAnimationOffset() const override { return AnimationOffset; }
+	virtual void SetAnimationEnabled(bool AnimationEnabledIn) override { bAnimationEnabled = AnimationEnabledIn; }
 	virtual bool GetAnimationEnabled() const override { return bAnimationEnabled; }
 	virtual TArray<FModuleInputSetup> GetInputConfig() const override { return InputConfig; }
 	virtual int32 GetAnimationSetupIndex() const override { return AnimationSetupIndex; }
diff --git a/Engine/Source/Runtime/Experimental/ChaosVehicles/ChaosVehiclesCore/Public/SimModule/SimulationModuleBase.h b/Engine/Source/Runtime/Experimental/ChaosVehicles/ChaosVehiclesCore/Public/SimModule/SimulationModuleBase.h
--- a/Engine/Source/Runtime/Experimental/ChaosVehicles/ChaosVehiclesCore/Public/SimModule/SimulationModuleBase.h
+++ b/Engine/Source/Runtime/Experimental/ChaosVehicles/ChaosVehiclesCore/Public/SimModule/SimulationModuleBase.h
@@ -300,6 +300,7 @@ namespace Chaos
 			, LocalAngularVelocity(FVector::ZeroVector)
 			, bClustered(true)
 			, bAnimationEnabled(true)
+			, AnimationOffset(FVector::ZeroVector)
 			, AppliedForce(FVector::ZeroVector)
 			, Guid(INDEX_NONE)
 			, CachedParticle(nullptr)
