[ChaosModularVehicles] Reimplement nettokens for Input without depenency for Iris in Core.

diff --git a/Engine/Plugins/Experimental/ChaosModularVehicle/Source/ChaosModularVehicleEngine/Private/ChaosModularVehicle/ChaosSimModuleManager.cpp b/Engine/Plugins/Experimental/ChaosModularVehicle/Source/ChaosModularVehicleEngine/Private/ChaosModularVehicle/ChaosSimModuleManager.cpp
--- a/Engine/Plugins/Experimental/ChaosModularVehicle/Source/ChaosModularVehicleEngine/Private/ChaosModularVehicle/ChaosSimModuleManager.cpp
+++ b/Engine/Plugins/Experimental/ChaosModularVehicle/Source/ChaosModularVehicleEngine/Private/ChaosModularVehicle/ChaosSimModuleManager.cpp
@@ -1,11 +1,11 @@
 //// Copyright Epic Games, Inc. All Rights Reserved.
 //
 #include "ChaosModularVehicle/ChaosSimModuleManager.h"
+#include "ChaosModularVehicle/ModuleInputTokenStore.h"
 #include "ChaosModularVehicle/ModularVehicleBaseComponent.h"
 #include "PBDRigidsSolver.h"
 #include "GameFramework/HUD.h" // for ShowDebugInfo
 #include "Physics/Experimental/PhysScene_Chaos.h"
-
 TMap<FPhysScene*, FChaosSimModuleManager*> FChaosSimModuleManager::SceneToModuleManagerMap;
 
 FDelegateHandle FChaosSimModuleManager::OnPostWorldInitializationHandle;
@@ -120,7 +120,12 @@ void FChaosSimModuleManager::RegisterNetTokenDataStores(UNetDriver* InNetDriver)
 	{
 		if (UE::Net::FNetTokenStore* TokenStore = InNetDriver->GetNetTokenStore())
 		{
+			typedef UE::Net::TStructNetTokenDataStore<FModuleInputNetTokenData> FModuleInputNetTokenStore;
 			typedef UE::Net::TStructNetTokenDataStore<FNetworkModularVehicleStateNetTokenData> FNetworkModularVehicleStateNetTokenStore;
+			if (!TokenStore->GetDataStore<FModuleInputNetTokenStore>())
+			{
+				TokenStore->CreateAndRegisterDataStore<FModuleInputNetTokenStore>();
+			}
 			if (!TokenStore->GetDataStore<FNetworkModularVehicleStateNetTokenStore>())
 			{
 				TokenStore->CreateAndRegisterDataStore<FNetworkModularVehicleStateNetTokenStore>();
diff --git a/Engine/Plugins/Experimental/ChaosModularVehicle/Source/ChaosModularVehicleEngine/Private/ChaosModularVehicle/ChaosSimModuleManagerAsyncCallback.cpp b/Engine/Plugins/Experimental/ChaosModularVehicle/Source/ChaosModularVehicleEngine/Private/ChaosModularVehicle/ChaosSimModuleManagerAsyncCallback.cpp
--- a/Engine/Plugins/Experimental/ChaosModularVehicle/Source/ChaosModularVehicleEngine/Private/ChaosModularVehicle/ChaosSimModuleManagerAsyncCallback.cpp
+++ b/Engine/Plugins/Experimental/ChaosModularVehicle/Source/ChaosModularVehicleEngine/Private/ChaosModularVehicle/ChaosSimModuleManagerAsyncCallback.cpp
@@ -4,7 +4,9 @@
 
 #include "ChaosModularVehicle/ModularVehicleBaseComponent.h"
 #include "ChaosModularVehicle/ModularVehicleSimulationCU.h"
+#include "ChaosModularVehicle/ModuleInputTokenStore.h"
 #include "PBDRigidsSolver.h"
+#include "ChaosModularVehicle/ModuleInputTokenStore.h"
 #include "PhysicsProxy/GeometryCollectionPhysicsProxy.h"
 #include "PhysicsProxy/ClusterUnionPhysicsProxy.h"
 #include "SimModule/ModuleFactoryRegister.h"
@@ -24,6 +26,7 @@ namespace ChaosModularVehicleCVars
 	FAutoConsoleVariableRef EnableStateNetSerializeDebugPrinting(TEXT("p.ModularVehicle.EnableStateNetSerializeDebugPrinting"), bEnableStateNetSerializeDebugPrinting, TEXT("Enable/Disable debug logging during NetSerialization. Default: false"));
 };
 
+UE_NET_IMPLEMENT_NAMED_NETTOKEN_STRUCT_SERIALIZERS(ModuleInputNetTokenData)
 UE_NET_IMPLEMENT_NAMED_NETTOKEN_STRUCT_SERIALIZERS(NetworkModularVehicleStateNetTokenData)
 
 FName FChaosSimModuleManagerAsyncCallback::GetFNameForStatId() const
@@ -258,10 +261,46 @@ bool FNetworkModularVehicleInputs::NetSerialize(FArchive& Ar, class UPackageMap*
 
 	Ar.SerializeBits(&VehicleInputs.Reverse, 1);
 	Ar.SerializeBits(&VehicleInputs.KeepAwake, 1);
-	
+
+	FNetworkModularVehicleInputs* DeltaSource = static_cast<FNetworkModularVehicleInputs*>(DeltaSourceData);
+	if (DeltaSource && DeltaSource->ServerFrame != -1 && ChaosModularVehicleCVars::bEnableStateReducedBandwidth)
+	{
+		using namespace UE::Net;
+		bOutSuccess = true;
+		
+		TArray<FModuleInputValue>& InputValues = VehicleInputs.Container.AccessInputValues();
+		const TArray<FModuleInputValue>& PreviousInputValues = DeltaSource->VehicleInputs.Container.AccessInputValues();
+		
+		FModuleInputNetTokenData InputStateData;
+		InputStateData.Init(InputValues);
+		TStructNetTokenDataStoreHelper<FModuleInputNetTokenData>::NetSerializeAndExportToken(Ar,Map,InputStateData);
+
+		uint32 Number = InputStateData.Types.Num();
+		if (Ar.IsLoading())
+		{
+			InputValues.SetNum(Number);
+		}
+
+		for (uint32 I = 0; I < Number; I++)
+		{
+			InputValues[I].ConvertToType(static_cast<EModuleInputValueType>(InputStateData.Types[I]));
+			InputValues[I].SetApplyInputDecay(InputStateData.DecayValues[I]);
+			if (PreviousInputValues.Num() == InputValues.Num())
+			{
+				InputValues[I].DeltaNetSerialize(Ar, Map, bOutSuccess, PreviousInputValues[I]);
+			}
+			else
+			{
+				InputValues[I].Serialize(Ar, Map, bOutSuccess);
+			}
+		
+		}
+	}
+	else
 	{
 		VehicleInputs.Container.Serialize(Ar, Map, bOutSuccess);
 	}
+	
 
 	bOutSuccess = true;
 	return bOutSuccess;
diff --git a/Engine/Plugins/Experimental/ChaosModularVehicle/Source/ChaosModularVehicleEngine/Public/ChaosModularVehicle/ModuleInputTokenStore.h b/Engine/Plugins/Experimental/ChaosModularVehicle/Source/ChaosModularVehicleEngine/Public/ChaosModularVehicle/ModuleInputTokenStore.h
new file mode 100644
--- /dev/null
+++ b/Engine/Plugins/Experimental/ChaosModularVehicle/Source/ChaosModularVehicleEngine/Public/ChaosModularVehicle/ModuleInputTokenStore.h
@@ -0,0 +1,39 @@
+// Copyright Epic Games, Inc. All Rights Reserved.
+
+#pragma once
+
+#include "Iris/ReplicationSystem/NetTokenStructDefines.h"
+#include "Iris/ReplicationSystem/StructNetTokenDataStore.h"
+#include "Net/Iris/ReplicationSystem/StructNetTokenDataStoreHelper.h"
+#include "SimModule/ModuleInput.h"
+#include "ModuleInputTokenStore.generated.h"
+
+USTRUCT()
+struct FModuleInputNetTokenData
+{
+	GENERATED_BODY()
+	UPROPERTY()
+	TArray<uint8> Types;
+	UPROPERTY()
+	TArray<bool> DecayValues;
+	
+	UE_NET_NETTOKEN_GENERATED_BODY(ModuleInputNetTokenData, CHAOSMODULARVEHICLEENGINE_API)
+	
+	CHAOSMODULARVEHICLEENGINE_API uint64 GetUniqueKey() const
+	{
+		uint64 HashOfTypes = CityHash64((const char*)Types.GetData(),sizeof(EModuleInputValueType)*Types.Num());
+		uint64 HashOfDecayValues = CityHash64((const char*)DecayValues.GetData(),sizeof(bool)*DecayValues.Num());
+		return HashOfTypes ^ HashOfDecayValues;
+	}
+
+	CHAOSMODULARVEHICLEENGINE_API void Init(const TArray<FModuleInputValue>& ModuleInputs)
+	{
+		for (int32 Idx = 0; Idx < ModuleInputs.Num(); Idx++)
+		{
+			Types.Add(static_cast<uint8>(ModuleInputs[Idx].GetValueType()));
+			DecayValues.Add(ModuleInputs[Idx].ShouldApplyInputDecay());
+		}
+	}
+};
+
+UE_NET_DECLARE_NAMED_NETTOKEN_STRUCT_SERIALIZERS(ModuleInputNetTokenData, CHAOSMODULARVEHICLEENGINE_API);
\ No newline at end of file
diff --git a/Engine/Source/Runtime/Experimental/ChaosVehicles/ChaosVehiclesCore/Private/SimModule/ModuleInput.cpp b/Engine/Source/Runtime/Experimental/ChaosVehicles/ChaosVehiclesCore/Private/SimModule/ModuleInput.cpp
--- a/Engine/Source/Runtime/Experimental/ChaosVehicles/ChaosVehiclesCore/Private/SimModule/ModuleInput.cpp
+++ b/Engine/Source/Runtime/Experimental/ChaosVehicles/ChaosVehiclesCore/Private/SimModule/ModuleInput.cpp
@@ -1,10 +1,18 @@
 // Copyright Epic Games, Inc. All Rights Reserved.
 
 #include "SimModule/ModuleInput.h"
+
+#include "HAL/IConsoleManager.h"
 #include "Net/Core/NetBitArray.h"
 
 TArray<struct FModuleInputSetup>* FScopedModuleInputInitializer::InitSetupData = nullptr;
 
+namespace ChaosModularVehicleCVars
+{
+	bool bEnableInputSendingTypeInfo = true;
+	FAutoConsoleVariableRef EnableInputSendingTypeInfo(TEXT("p.ModularVehicle.EnableInputSendingTypeInfo"), bEnableInputSendingTypeInfo, TEXT("Enable/Disable Input sending type info. Default: true"));
+};
+
 DEFINE_LOG_CATEGORY(LogModularInput);
 
 bool FModuleInputValue::IsNonZero(float Tolerance) const
@@ -116,8 +124,11 @@ void FModuleInputValue::Serialize(FArchive& Ar, UPackageMap* Map, bool& bOutSucc
 
 	bool bIsNonZero = Ar.IsLoading() ? false : IsQuantizedNonZero();
 	Ar.SerializeBits(&bIsNonZero, 1);
-	Ar.SerializeBits(&ValueType, 3);
-	Ar.SerializeBits(&bApplyInputDecay, 1);
+	if (!ChaosModularVehicleCVars::bEnableInputSendingTypeInfo)
+	{
+		Ar.SerializeBits(&ValueType, 3);
+		Ar.SerializeBits(&bApplyInputDecay, 1);
+	}
 
 	if (bIsNonZero)
 	{
@@ -160,6 +171,31 @@ void FModuleInputValue::Serialize(FArchive& Ar, UPackageMap* Map, bool& bOutSucc
 	bOutSuccess = true;
 }
 
+void FModuleInputValue::DeltaNetSerialize(FArchive& Ar, UPackageMap* Map, bool& bOutSuccess, const FModuleInputValue& PreviousInputValue)
+{
+	bool bHasSamePreviousValue = Value == PreviousInputValue.Value;
+	bool bHasSamePreviousType = ValueType == PreviousInputValue.ValueType;
+	bool bHasSamePreviousValueInt = ValueInt == PreviousInputValue.ValueInt;
+	bool bIsSameAsPrevious = bHasSamePreviousValue && bHasSamePreviousType && bHasSamePreviousValueInt;
+	Ar.SerializeBits(&bIsSameAsPrevious, 1);
+	if (bIsSameAsPrevious)
+	{
+		bOutSuccess = true;
+		if (Ar.IsSaving())
+		{
+			return;
+		}
+
+		ConvertToType(PreviousInputValue.ValueType);
+		Value = PreviousInputValue.Value;
+		ValueInt = PreviousInputValue.ValueInt;
+		bApplyInputDecay = PreviousInputValue.bApplyInputDecay;
+		return;
+	}
+
+	Serialize(Ar, Map, bOutSuccess);
+}
+
 void FModuleInputValue::Lerp(const FModuleInputValue& Min, const FModuleInputValue& Max, float Alpha)
 {
 	switch (GetValueType())
diff --git a/Engine/Source/Runtime/Experimental/ChaosVehicles/ChaosVehiclesCore/Public/SimModule/ModuleInput.h b/Engine/Source/Runtime/Experimental/ChaosVehicles/ChaosVehiclesCore/Public/SimModule/ModuleInput.h
--- a/Engine/Source/Runtime/Experimental/ChaosVehicles/ChaosVehiclesCore/Public/SimModule/ModuleInput.h
+++ b/Engine/Source/Runtime/Experimental/ChaosVehicles/ChaosVehiclesCore/Public/SimModule/ModuleInput.h
@@ -103,7 +103,7 @@ namespace ModularQuantize
 	}
 
 	template<int32 MaxValue, uint32 NumBits, typename T UE_REQUIRES(std::is_floating_point_v<T>&& NumBits < 32)>
-	bool WriteCompressedFloat(const T Value, FArchive& Ar)
+bool WriteCompressedFloat(const T Value, FArchive& Ar)
 	{
 		using Details = ModularQuantize::TCompressedFloatDetails<MaxValue, NumBits>;
 
@@ -150,6 +150,18 @@ namespace ModularQuantize
 		ModularQuantize::ToCompressedFloat<MaxValue, NumBits>(Value, CompressedValue);
 		ModularQuantize::FromCompressedFloat<MaxValue, NumBits>(CompressedValue, Value);
 	}
+	
+	template<int32 MaxValue, uint32 NumBits, typename T UE_REQUIRES(std::is_floating_point_v<T>&& NumBits < 32)>
+	bool QuantizedIsNearlyEqual(const T& Left,const T& Right)
+	{
+		uint32 CompressedValue = 0;
+		T QuantLeft, QuantRight;
+		ModularQuantize::ToCompressedFloat<MaxValue, NumBits>(Left, CompressedValue);
+		ModularQuantize::FromCompressedFloat<MaxValue, NumBits>(CompressedValue, QuantLeft);
+		ModularQuantize::ToCompressedFloat<MaxValue, NumBits>(Right, CompressedValue);
+		ModularQuantize::FromCompressedFloat<MaxValue, NumBits>(CompressedValue, QuantRight);
+		return FMath::IsNearlyEqual(QuantLeft, QuantRight);
+	}
 
 }
 
@@ -356,6 +368,7 @@ public:
 
 	// Serialize values
 	void Serialize(FArchive& Ar, UPackageMap* Map, bool& bOutSuccess);
+	void DeltaNetSerialize(FArchive& Ar, UPackageMap* Map, bool& bOutSuccess,const FModuleInputValue& PreviousInputValue);
 
 	void Lerp(const FModuleInputValue& Min, const FModuleInputValue& Max, float Alpha);
 
@@ -368,7 +381,7 @@ public:
 
 	/** During physics resimulation, apply decay for this input while it's being extrapolated */
 	void SetApplyInputDecay(const bool bInApplyInputDecay) { bApplyInputDecay = bInApplyInputDecay; }
-	bool ShouldApplyInputDecay() { return  bApplyInputDecay; }
+	bool ShouldApplyInputDecay() const { return  bApplyInputDecay; }
 
 protected:
 	void SetMagnitude(float NewSize);
@@ -628,6 +641,8 @@ public:
 	/**  Decay input during resimulation by DecayAmount which increases over resimulation frames from 0.0 -> 1.0 when the input is being reused */
 	void Decay(const float DecayAmount);
 
+	TArray<FModuleInputValue>& AccessInputValues() { return InputValues; }
+
 private:
 	UPROPERTY()
 	TArray<FModuleInputValue> InputValues;
