The deserialization was broken for the case when parts of vehicle fall off.

diff --git a/Engine/Plugins/Experimental/ChaosModularVehicle/Source/ChaosModularVehicleEngine/Private/ChaosModularVehicle/ChaosSimModuleManagerAsyncCallback.cpp b/Engine/Plugins/Experimental/ChaosModularVehicle/Source/ChaosModularVehicleEngine/Private/ChaosModularVehicle/ChaosSimModuleManagerAsyncCallback.cpp
--- a/Engine/Plugins/Experimental/ChaosModularVehicle/Source/ChaosModularVehicleEngine/Private/ChaosModularVehicle/ChaosSimModuleManagerAsyncCallback.cpp
+++ b/Engine/Plugins/Experimental/ChaosModularVehicle/Source/ChaosModularVehicleEngine/Private/ChaosModularVehicle/ChaosSimModuleManagerAsyncCallback.cpp
@@ -405,17 +405,30 @@ bool FNetworkModularVehicleStates::NetSerialize(FArchive& Ar, class UPackageMap*
 
 				if (I <= static_cast<uint32>(ModuleData.Num()) && ModuleData[I].IsValid())
 				{
-					check(ModuleTypeHash == Chaos::FModuleFactoryRegister::GetModuleHash(ModuleData[I]->GetSimType()));
-
-					const bool bHasSerializedData = ModulesBitArray.IsBitSet(I);
-					if (bHasSerializedData)
+					bool bFound = false;
+					for (int J = 0; J < ModuleData.Num(); J++)
 					{
-						ModuleData[I]->Serialize(Ar);
-					}
-					else
-					{
-						ModuleData[I]->ApplyDefaultState();
+						if (SimArrayIndex == ModuleData[J].Get()->SimArrayIndex)
+						{
+							bFound = true;
+
+							check(ModuleTypeHash == Chaos::FModuleFactoryRegister::GetModuleHash(ModuleData[J]->GetSimType()));
+
+							const bool bHasSerializedData = ModulesBitArray.IsBitSet(I);
+							if (bHasSerializedData)
+							{
+								ModuleData[J]->Serialize(Ar);
+							}
+							else
+							{
+								ModuleData[J]->ApplyDefaultState();
+							}
+
+							break;
+						}
 					}
+
+					ensureMsgf (bFound, TEXT("Module  with SimArrayIndex %d not found"), SimArrayIndex);
 				}
 			}
 		}
