Modular Vehicle - Handle kinda rare bad case gracefully for now

diff --git a/Engine/Plugins/Experimental/ChaosModularVehicle/Source/ChaosModularVehicleEngine/Private/ChaosModularVehicle/ModularVehicleBaseComponent.cpp b/Engine/Plugins/Experimental/ChaosModularVehicle/Source/ChaosModularVehicleEngine/Private/ChaosModularVehicle/ModularVehicleBaseComponent.cpp
--- a/Engine/Plugins/Experimental/ChaosModularVehicle/Source/ChaosModularVehicleEngine/Private/ChaosModularVehicle/ModularVehicleBaseComponent.cpp
+++ b/Engine/Plugins/Experimental/ChaosModularVehicle/Source/ChaosModularVehicleEngine/Private/ChaosModularVehicle/ModularVehicleBaseComponent.cpp
@@ -279,18 +279,23 @@ void UModularVehicleBaseComponent::CreateAssociatedSimComponents(USceneComponent
 		return;
 	}
 
-	if (RootPhysicsObject == nullptr)
-	{
-		CacheRootPhysicsObject(GetPhysicsProxy());
-	}
-
 	UE_LOG(LogModularBase, Log, TEXT("CreateAssociatedSimComponents Attaching %s, TransformIndex %d"), *AttachedComponent->GetName(), InTransformIndex);
 
 	TArray<Chaos::FClusterUnionChildData> ChildParticles = ClusterUnionComponent->GetPhysicsProxy()->GetSyncedData_External().ChildParticles;
 
+	if (InTransformIndex >= ChildParticles.Num())
+	{
+		// #TODO: look into the underlying cause of this indexing issue, possibly not an easy fix
+		ensureMsgf(false, TEXT("Unexpected state - likely due to adding and removing components on same frame leading to an array indexing issue") );
+		return;
+	}
+
 	int TransformIndex = InTransformIndex;
+	if (RootPhysicsObject == nullptr)
+	{
+		CacheRootPhysicsObject(GetPhysicsProxy());
+	}
 
-	ensure(TransformIndex < ChildParticles.Num());
 
 	if (IVehicleSimBaseComponentInterface* ComponentInterface = Cast<IVehicleSimBaseComponentInterface>(AttachedComponent))
 	{
