[ModularVehicle] Refactor NetSerializers to consider DeltaState to be "valid".

diff --git a/Engine/Plugins/Experimental/ChaosModularVehicle/Source/ChaosModularVehicleEngine/Private/ChaosModularVehicle/ChaosSimModuleManagerAsyncCallback.cpp b/Engine/Plugins/Experimental/ChaosModularVehicle/Source/ChaosModularVehicleEngine/Private/ChaosModularVehicle/ChaosSimModuleManagerAsyncCallback.cpp
--- a/Engine/Plugins/Experimental/ChaosModularVehicle/Source/ChaosModularVehicleEngine/Private/ChaosModularVehicle/ChaosSimModuleManagerAsyncCallback.cpp
+++ b/Engine/Plugins/Experimental/ChaosModularVehicle/Source/ChaosModularVehicleEngine/Private/ChaosModularVehicle/ChaosSimModuleManagerAsyncCallback.cpp
@@ -7,6 +7,7 @@
 #include "ChaosModularVehicle/ModuleInputTokenStore.h"
 #include "PBDRigidsSolver.h"
 #include "ChaosModularVehicle/ModuleInputTokenStore.h"
+#include "GeometryCollection/GeometryCollectionParticlesData.h"
 #include "PhysicsProxy/GeometryCollectionPhysicsProxy.h"
 #include "PhysicsProxy/ClusterUnionPhysicsProxy.h"
 #include "SimModule/ModuleFactoryRegister.h"
@@ -21,7 +22,9 @@ DECLARE_CYCLE_STAT(TEXT("AsyncCallback:OnContactModification_Internal"), STAT_As
 namespace ChaosModularVehicleCVars
 {
 	bool bEnableStateReducedBandwidth = false;
-	FAutoConsoleVariableRef EnableStateReducedBandwidth(TEXT("p.ModularVehicle.EnableStateReducedBandwidth"), bEnableStateReducedBandwidth, TEXT("Enable/Disable NetTokens and DeltaSerialization path for Modular Vehicles. Default: false"));
+	FAutoConsoleVariableRef EnableStateReducedBandwidth(TEXT("p.ModularVehicle.EnableStateReducedBandwidth"), bEnableStateReducedBandwidth, TEXT("Enable/Disable NetTokens and DeltaSerialization path for State of Modular Vehicles. Default: false"));
+	bool bEnableInputReducedBandwidth = false;
+	FAutoConsoleVariableRef EnableInputReducedBandwidth(TEXT("p.ModularVehicle.EnableInputReducedBandwidth"), bEnableInputReducedBandwidth, TEXT("Enable/Disable NetTokens and DeltaSerialization path for Input of Modular Vehicles. Default: false"));
 	bool bEnableStateNetSerializeDebugPrinting = false;
 	FAutoConsoleVariableRef EnableStateNetSerializeDebugPrinting(TEXT("p.ModularVehicle.EnableStateNetSerializeDebugPrinting"), bEnableStateNetSerializeDebugPrinting, TEXT("Enable/Disable debug logging during NetSerialization. Default: false"));
 };
@@ -261,9 +264,9 @@ bool FNetworkModularVehicleInputs::NetSerialize(FArchive& Ar, class UPackageMap*
 
 	Ar.SerializeBits(&VehicleInputs.Reverse, 1);
 	Ar.SerializeBits(&VehicleInputs.KeepAwake, 1);
-
+	bOutSuccess = true;
 	FNetworkModularVehicleInputs* DeltaSource = static_cast<FNetworkModularVehicleInputs*>(DeltaSourceData);
-	if (DeltaSource && DeltaSource->ServerFrame != -1 && ChaosModularVehicleCVars::bEnableStateReducedBandwidth)
+	if (DeltaSource && ChaosModularVehicleCVars::bEnableInputReducedBandwidth)
 	{
 		using namespace UE::Net;
 		bOutSuccess = true;
@@ -273,8 +276,13 @@ bool FNetworkModularVehicleInputs::NetSerialize(FArchive& Ar, class UPackageMap*
 		
 		FModuleInputNetTokenData InputStateData;
 		InputStateData.Init(InputValues);
-		TStructNetTokenDataStoreHelper<FModuleInputNetTokenData>::NetSerializeAndExportToken(Ar,Map,InputStateData);
-
+		const bool bNetTokenSuccess = TStructNetTokenDataStoreHelper<FModuleInputNetTokenData>::NetSerializeAndExportToken(Ar,Map,InputStateData);
+		if (!bNetTokenSuccess)
+		{
+			bOutSuccess = false;
+			return bOutSuccess;
+		}
+		
 		uint32 Number = InputStateData.Types.Num();
 		if (Ar.IsLoading())
 		{
@@ -291,7 +299,9 @@ bool FNetworkModularVehicleInputs::NetSerialize(FArchive& Ar, class UPackageMap*
 			}
 			else
 			{
-				InputValues[I].Serialize(Ar, Map, bOutSuccess);
+				bOutSuccess = false;
+				//Fail case.
+				InputValues[I].DeltaNetSerialize(Ar, Map, bOutSuccess, InputValues[I]);
 			}
 		
 		}
@@ -301,8 +311,6 @@ bool FNetworkModularVehicleInputs::NetSerialize(FArchive& Ar, class UPackageMap*
 		VehicleInputs.Container.Serialize(Ar, Map, bOutSuccess);
 	}
 	
-
-	bOutSuccess = true;
 	return bOutSuccess;
 }
 
@@ -359,9 +367,9 @@ void FNetworkModularVehicleInputs::DecayData(float DecayAmount)
 
 bool FNetworkModularVehicleStates::NetSerialize(FArchive& Ar, class UPackageMap* Map, bool& bOutSuccess)
 {
-	if (ChaosModularVehicleCVars::bEnableStateReducedBandwidth)
+	if (ChaosModularVehicleCVars::bEnableStateReducedBandwidth && DeltaSourceData != nullptr)
 	{
-		return NetSerializeReduced(Ar, Map, bOutSuccess);
+		return DeltaNetSerialize(Ar, Map, bOutSuccess);
 	}
 	
 	using namespace UE::Net;
@@ -405,30 +413,17 @@ bool FNetworkModularVehicleStates::NetSerialize(FArchive& Ar, class UPackageMap*
 
 				if (I <= static_cast<uint32>(ModuleData.Num()) && ModuleData[I].IsValid())
 				{
-					bool bFound = false;
-					for (int J = 0; J < ModuleData.Num(); J++)
+					check(ModuleTypeHash == Chaos::FModuleFactoryRegister::GetModuleHash(ModuleData[I]->GetSimType()));
+
+					const bool bHasSerializedData = ModulesBitArray.IsBitSet(I);
+					if (bHasSerializedData)
 					{
-						if (SimArrayIndex == ModuleData[J].Get()->SimArrayIndex)
-						{
-							bFound = true;
-
-							check(ModuleTypeHash == Chaos::FModuleFactoryRegister::GetModuleHash(ModuleData[J]->GetSimType()));
-
-							const bool bHasSerializedData = ModulesBitArray.IsBitSet(I);
-							if (bHasSerializedData)
-							{
-								ModuleData[J]->Serialize(Ar);
-							}
-							else
-							{
-								ModuleData[J]->ApplyDefaultState();
-							}
-
-							break;
-						}
+						ModuleData[I]->Serialize(Ar);
+					}
+					else
+					{
+						ModuleData[I]->ApplyDefaultState();
 					}
-
-					ensureMsgf (bFound, TEXT("Module  with SimArrayIndex %d not found"), SimArrayIndex);
 				}
 			}
 		}
@@ -467,118 +462,59 @@ bool FNetworkModularVehicleStates::NetSerialize(FArchive& Ar, class UPackageMap*
 	return true;
 }
 
-bool FNetworkModularVehicleStates::NetSerializeReduced(FArchive& Ar, class UPackageMap* Map, bool& bOutSuccess)
+bool FNetworkModularVehicleStates::DeltaNetSerialize(FArchive& Ar, UPackageMap* Map, bool& bOutSuccess)
 {
 	const bool bPrintDebugInfo = ChaosModularVehicleCVars::bEnableStateNetSerializeDebugPrinting;
+	FBitReader* BitReader = (FBitReader*)&Ar;
+	FBitWriter* BitWriter = (FBitWriter*)&Ar;
 	using namespace UE::Net;
 
 	FNetworkPhysicsData::SerializeFrames(Ar);
-	
-	FNetworkModularVehicleStateNetTokenData ModuleStateData;
-	ModuleStateData.Init(ModuleData);
-	TStructNetTokenDataStoreHelper<FNetworkModularVehicleStateNetTokenData>::NetSerializeAndExportToken(Ar,Map,ModuleStateData);
 
-	if (Ar.IsSaving())
-	{
-		check(ModuleStateData.Hashes.Num() == ModuleData.Num());
-	}
 	FNetworkModularVehicleStates* DeltaSource = static_cast<FNetworkModularVehicleStates*>(DeltaSourceData);
-
-	UE_LOG(LogTemp, Warning, TEXT("====NetSerialize %s. ServerFrame: %d. DeltaFrame: %d"),Ar.IsSaving()?TEXT("Save"):TEXT("Load"),ServerFrame,DeltaSource?DeltaSource->ServerFrame:-1);
+	UE_CLOG(bPrintDebugInfo,LogModularVehicleSim, Warning, TEXT("====DeltaNetSerialize Saving: %d. ServerFrame: %d. DeltaSource_ServerFrame: %d Starting Bit: %lld"), Ar.IsSaving(), ServerFrame, DeltaSource->ServerFrame, Ar.IsSaving()?BitWriter->GetNumBits():BitReader->GetPosBits());
 	
-	if (DeltaSource && DeltaSource->ServerFrame != -1 && DeltaSource->ModuleData.Num() == ModuleStateData.Hashes.Num() )
+	FNetworkModularVehicleStateNetTokenData VehicleStateData;
+	VehicleStateData.Init(ModuleData);
+	const bool bNetTokenSuccess = TStructNetTokenDataStoreHelper<FNetworkModularVehicleStateNetTokenData>::NetSerializeAndExportToken(Ar,Map,VehicleStateData);
+	if (!bNetTokenSuccess)
 	{
-		return DeltaNetSerialize(Ar,Map,bOutSuccess,ModuleStateData);
+		bOutSuccess = false;
+		return bOutSuccess;
 	}
-	FBitReader* BitReader = (FBitReader*)&Ar;
-	FBitWriter* BitWriter = (FBitWriter*)&Ar;
-	UE_CLOG(bPrintDebugInfo,LogTemp, Warning, TEXT("====NetSerialize Saving: %d. ServerFrame: %d. Starting Bit: %lld"),Ar.IsSaving(),ServerFrame,Ar.IsSaving()?BitWriter->GetNumBits():BitReader->GetPosBits());
-	uint32 NumNetModules = ModuleStateData.Hashes.Num();
 
-	// Array of bits to mark which modules to serialize or not
-	if (Ar.IsLoading())
+	auto GetDeltaDataHelper = [&](int32 InIdx, uint32 InModuleTypeHash, int32 InSimArrayIndex)
 	{
-		if (NumNetModules != ModuleData.Num())
-		{
-			ModuleData.SetNum(NumNetModules);
-		}
-		if(bPrintDebugInfo)
-		{
-			FString BitString = "";
-			for (uint32 I = 0; I < NumNetModules; I++)
-			{
-				BitString += ModuleStateData.ModuleShouldSerialize[I] ? "1" : "0";
-			}
-			UE_CLOG(bPrintDebugInfo,LogTemp, Warning, TEXT("==NetSerialize LOADING. Using ModuleShouldSerialize: %s"),*BitString);
-		}
-		for (uint32 I = 0; I < NumNetModules; I++)
+		TSharedPtr<Chaos::FModuleNetData> DeltaData = nullptr;
+		if (DeltaSource->ModuleData.IsValidIndex(InIdx))
 		{
-			if (NumNetModules > 0)
+			const uint32 DeltaModuleHash = Chaos::FModuleFactoryRegister::GetModuleHash(DeltaSource->ModuleData[InIdx]->GetSimType());
+			if (DeltaModuleHash == InModuleTypeHash)
 			{
-				uint32 ModuleTypeHash = ModuleStateData.Hashes[I];
-				int32 SimArrayIndex = ModuleStateData.Indexes[I];
-
-				if (TSharedPtr<Chaos::FModuleNetData> Data = Chaos::FModuleFactoryRegister::Get().GenerateNetData(ModuleTypeHash, SimArrayIndex))
-				{
-					UE_CLOG(bPrintDebugInfo,LogTemp, Warning, TEXT("==NetSerialize LOADING. ModuleData: %d STA. Bit: %lld"),I,Ar.IsSaving()?BitWriter->GetNumBits():BitReader->GetPosBits());
-					ModuleData[I] = Data;
-					if (ModuleStateData.ModuleShouldSerialize[I])
-					{
-						ModuleData[I]->Serialize(Ar);
-					}
-					else
-					{
-						ModuleData[I]->ApplyDefaultState();
-					}
-					UE_CLOG(bPrintDebugInfo,LogTemp, Warning, TEXT("==NetSerialize LOADING. ModuleData: %d STA. Bit: %lld"),I,Ar.IsSaving()?BitWriter->GetNumBits():BitReader->GetPosBits());
-				}
+				DeltaData = DeltaSource->ModuleData[InIdx];
 			}
 		}
-	}
-	else
-	{
-		if (bPrintDebugInfo)
+		if (DeltaData == nullptr)
 		{
-			FString BitString = "";
-			for (uint32 I = 0; I < NumNetModules; I++)
+			UE_CLOG(bPrintDebugInfo,LogModularVehicleSim, Warning, TEXT("==DeltaNetSerialize Generating Default Data for DeltaData Module %u"), InModuleTypeHash);
+			DeltaData = Chaos::FModuleFactoryRegister::Get().GenerateNetData(InModuleTypeHash, InSimArrayIndex);
+			if (DeltaData == nullptr)
 			{
-				BitString += ModuleStateData.ModuleShouldSerialize[I] ? "1" : "0";
+				UE_LOG(LogModularVehicleSim, Error, TEXT("Unable to generate net data for delta source when delta is invalid"));
+				bOutSuccess = false;
 			}
-			UE_CLOG(bPrintDebugInfo,LogTemp, Warning, TEXT("==NetSerialize SAVING. Using ModuleShouldSerialize: %s"),*BitString)
-		}
-		for (uint32 I = 0; I < NumNetModules; I++)
-		{
-			UE_CLOG(bPrintDebugInfo,LogTemp, Warning, TEXT("==NetSerialize SAVING. ModuleData: %d STA. Bit: %lld"),I,Ar.IsSaving()?BitWriter->GetNumBits():BitReader->GetPosBits())
-			if (ModuleStateData.ModuleShouldSerialize[I])
+			else
 			{
-				ModuleData[I]->Serialize(Ar);
+				DeltaData->ApplyDefaultState();
 			}
-			UE_CLOG(bPrintDebugInfo,LogTemp, Warning, TEXT("==NetSerialize SAVING. ModuleData: %d END. Bit: %lld"),I,Ar.IsSaving()?BitWriter->GetNumBits():BitReader->GetPosBits())
 		}
-	}
-
-	bOutSuccess = true;
-	return true;
-}
-
-bool FNetworkModularVehicleStates::DeltaNetSerialize(FArchive& Ar, UPackageMap* Map, bool& bOutSuccess, FNetworkModularVehicleStateNetTokenData& VehicleStateData)
-{
-	const bool bPrintDebugInfo = ChaosModularVehicleCVars::bEnableStateNetSerializeDebugPrinting;
-	FBitReader* BitReader = (FBitReader*)&Ar;
-	FBitWriter* BitWriter = (FBitWriter*)&Ar;
-	using namespace UE::Net;
-	uint32 NumNetModules = VehicleStateData.Hashes.Num();
-	FNetworkModularVehicleStates* DeltaSource = static_cast<FNetworkModularVehicleStates*>(DeltaSourceData);
-	if (DeltaSource->ModuleData.Num() != NumNetModules)
-	{
-		return NetSerialize(Ar, Map, bOutSuccess);
-	}
-	UE_CLOG(bPrintDebugInfo,LogTemp, Warning, TEXT("====DeltaNetSerialize Saving: %d. ServerFrame: %d. DeltaSource_ServerFrame: %d Starting Bit: %lld"), Ar.IsSaving(), ServerFrame, DeltaSource->ServerFrame, Ar.IsSaving()?BitWriter->GetNumBits():BitReader->GetPosBits());
+		return DeltaData;
+	};
 	
-
+	const uint32 NumNetModules = VehicleStateData.Hashes.Num();
+	bOutSuccess = true;
 	TMap<FName, uint32> SerializationStash;
 	SerializationStash.Add(StashServerFrameKey, ServerFrame);
-	
 	if (Ar.IsLoading())
 	{
 		if (NumNetModules != ModuleData.Num())
@@ -588,32 +524,34 @@ bool FNetworkModularVehicleStates::DeltaNetSerialize(FArchive& Ar, UPackageMap*
 		if(bPrintDebugInfo)
 		{
 			FString BitString = "";
-			for (uint32 I = 0; I < NumNetModules; I++)
+			for (uint32 NetModuleIdx = 0; NetModuleIdx < NumNetModules; ++NetModuleIdx)
 			{
-				BitString += VehicleStateData.ModuleShouldSerialize[I] ? "1" : "0";
+				BitString += VehicleStateData.ModuleShouldSerialize[NetModuleIdx] ? "1" : "0";
 			}
-			UE_CLOG(bPrintDebugInfo,LogTemp, Warning, TEXT("==DeltaNetSerialize LOADING. Using ModuleShouldSerialize: %s"),*BitString);
+			UE_CLOG(bPrintDebugInfo,LogModularVehicleSim, Warning, TEXT("==DeltaNetSerialize LOADING. Using ModuleShouldSerialize: %s"),*BitString);
 		}
 		for (uint32 I = 0; I < NumNetModules; I++)
 		{
-			int32 SimArrayIndex = VehicleStateData.Indexes[I];
-			uint32 ModuleTypeHash = VehicleStateData.Hashes[I];
+			const int32 SimArrayIndex = VehicleStateData.Indexes[I];
+			const uint32 ModuleTypeHash = VehicleStateData.Hashes[I];
+			
 			if (TSharedPtr<Chaos::FModuleNetData> Data = Chaos::FModuleFactoryRegister::Get().GenerateNetData(ModuleTypeHash, SimArrayIndex))
 			{
 				int32 StartBit = Ar.IsSaving()?BitWriter->GetNumBits():BitReader->GetPosBits();
-				UE_CLOG(bPrintDebugInfo,LogTemp, Warning, TEXT("==DeltaNetSerialize LOADING. ModuleData: %d STA. Bit: %lld"),I,Ar.IsSaving()?BitWriter->GetNumBits():BitReader->GetPosBits());
+				UE_CLOG(bPrintDebugInfo,LogModularVehicleSim, Warning, TEXT("==DeltaNetSerialize LOADING. ModuleData: %d STA. Bit: %lld"),I,Ar.IsSaving()?BitWriter->GetNumBits():BitReader->GetPosBits());
 				ModuleData[I] = Data;
 				check(ModuleTypeHash == Chaos::FModuleFactoryRegister::GetModuleHash(Data->GetSimType()));
 				if (VehicleStateData.ModuleShouldSerialize[I])
 				{
-					ModuleData[I]->DeltaSerializeWithStash(Ar, DeltaSource->ModuleData[I].Get(),SerializationStash);
+					TSharedPtr<Chaos::FModuleNetData> DeltaData = GetDeltaDataHelper(I, ModuleTypeHash, SimArrayIndex);
+					ModuleData[I]->DeltaSerializeWithStash(Ar, DeltaData.Get(),SerializationStash);
 				}
 				else
 				{
 					ModuleData[I]->ApplyDefaultState();
 				}
 				int32 EndBit = Ar.IsSaving()?BitWriter->GetNumBits():BitReader->GetPosBits();
-				UE_CLOG(bPrintDebugInfo,LogTemp, Warning, TEXT("==DeltaNetSerialize LOADING. ModuleData: %d END. Bit: %lld Total: %d Error: %d"),I,Ar.IsSaving()?BitWriter->GetNumBits():BitReader->GetPosBits(), EndBit-StartBit,Ar.IsError());
+				UE_CLOG(bPrintDebugInfo,LogModularVehicleSim, Warning, TEXT("==DeltaNetSerialize LOADING. ModuleData: %d END. Bit: %lld Total: %d Error: %d"),I,Ar.IsSaving()?BitWriter->GetNumBits():BitReader->GetPosBits(), EndBit-StartBit,Ar.IsError());
 			}
 		}
 	}
@@ -626,22 +564,25 @@ bool FNetworkModularVehicleStates::DeltaNetSerialize(FArchive& Ar, UPackageMap*
 			{
 				BitString += VehicleStateData.ModuleShouldSerialize[I] ? "1" : "0";
 			}
-			UE_CLOG(bPrintDebugInfo,LogTemp, Warning, TEXT("==DeltaNetSerialize SAVING. Using ModuleShouldSerialize: %s"),*BitString);
+			UE_CLOG(bPrintDebugInfo,LogModularVehicleSim, Warning, TEXT("==DeltaNetSerialize SAVING. Using ModuleShouldSerialize: %s"),*BitString);
 		}
 		for (uint32 I = 0; I < NumNetModules; I++)
 		{
 			int32 StartBit = Ar.IsSaving()?BitWriter->GetNumBits():BitReader->GetPosBits();
-			UE_CLOG(bPrintDebugInfo,LogTemp, Warning, TEXT("==DeltaNetSerialize SAVING. ModuleData: %d STA. Bit: %lld - %s"),I, Ar.IsSaving()?BitWriter->GetNumBits():BitReader->GetPosBits(), *DeltaSource->ModuleData[I]->GetSimType().ToString());
+			UE_CLOG(bPrintDebugInfo,LogModularVehicleSim, Warning, TEXT("==DeltaNetSerialize SAVING. ModuleData: %d STA. Bit: %lld - %s"),I, Ar.IsSaving()?BitWriter->GetNumBits():BitReader->GetPosBits(), *ModuleData[I]->GetSimType().ToString());
 			if (VehicleStateData.ModuleShouldSerialize[I])
 			{
-				ModuleData[I]->DeltaSerializeWithStash(Ar, DeltaSource->ModuleData[I].Get(), SerializationStash);
+				const int32 SimArrayIndex = VehicleStateData.Indexes[I];
+				const uint32 ModuleTypeHash = VehicleStateData.Hashes[I];
+				TSharedPtr<Chaos::FModuleNetData> DeltaData = GetDeltaDataHelper(I, ModuleTypeHash, SimArrayIndex);
+				ModuleData[I]->DeltaSerializeWithStash(Ar, DeltaData.Get(), SerializationStash);
 			}
 			int32 EndBit = Ar.IsSaving()?BitWriter->GetNumBits():BitReader->GetPosBits();
-			UE_CLOG(bPrintDebugInfo,LogTemp, Warning, TEXT("==DeltaNetSerialize SAVING. ModuleData: %d END. Bit: %lld Size: %d - %s"),I, Ar.IsSaving()?BitWriter->GetNumBits():BitReader->GetPosBits(), EndBit-StartBit, *DeltaSource->ModuleData[I]->GetSimType().ToString());
+			UE_CLOG(bPrintDebugInfo,LogModularVehicleSim, Warning, TEXT("==DeltaNetSerialize SAVING. ModuleData: %d END. Bit: %lld Size: %d - %s"),I, Ar.IsSaving()?BitWriter->GetNumBits():BitReader->GetPosBits(), EndBit-StartBit, *ModuleData[I]->GetSimType().ToString());
 		}
 	}
-	bOutSuccess = true;
-	return true;
+	
+	return bOutSuccess;
 }
 
 void FNetworkModularVehicleStates::ApplyData(UActorComponent* NetworkComponent) const
diff --git a/Engine/Plugins/Experimental/ChaosModularVehicle/Source/ChaosModularVehicleEngine/Public/ChaosModularVehicle/ChaosSimModuleManagerAsyncCallback.h b/Engine/Plugins/Experimental/ChaosModularVehicle/Source/ChaosModularVehicleEngine/Public/ChaosModularVehicle/ChaosSimModuleManagerAsyncCallback.h
--- a/Engine/Plugins/Experimental/ChaosModularVehicle/Source/ChaosModularVehicleEngine/Public/ChaosModularVehicle/ChaosSimModuleManagerAsyncCallback.h
+++ b/Engine/Plugins/Experimental/ChaosModularVehicle/Source/ChaosModularVehicleEngine/Public/ChaosModularVehicle/ChaosSimModuleManagerAsyncCallback.h
@@ -19,6 +19,7 @@
 class UModularVehicleComponent;
 class UModularVehicleBaseComponent;
 class FGeometryCollectionPhysicsProxy;
+class UPackageMap;
 
 DECLARE_STATS_GROUP(TEXT("ChaosSimModuleManager"), STATGROUP_ChaosSimModuleManager, STATGROUP_Advanced);
 
@@ -84,7 +85,7 @@ struct CHAOSMODULARVEHICLEENGINE_API FNetworkModularVehicleInputs : public FNetw
 	virtual void BuildData(const UActorComponent* NetworkComponent) override;
 
 	/**  Serialize data function that will be used to transfer the struct across the network */
-	bool NetSerialize(FArchive& Ar, class UPackageMap* Map, bool& bOutSuccess);
+	bool NetSerialize(FArchive& Ar, UPackageMap* Map, bool& bOutSuccess);
 
 	/** Interpolate the data in between two inputs */
 	virtual void InterpolateData(const FNetworkPhysicsData& MinData, const FNetworkPhysicsData& MaxData) override;
@@ -155,9 +156,8 @@ struct CHAOSMODULARVEHICLEENGINE_API FNetworkModularVehicleStates : public FNetw
 	virtual void BuildData(const UActorComponent* NetworkComponent) override;
 
 	/**  Serialize data function that will be used to transfer the struct across the network */
-	bool NetSerialize(FArchive& Ar, class UPackageMap* Map, bool& bOutSuccess);
-	bool NetSerializeReduced(FArchive& Ar, class UPackageMap* Map, bool& bOutSuccess);
-	bool DeltaNetSerialize(FArchive& Ar, class UPackageMap* Map, bool& BOutSuccess, FNetworkModularVehicleStateNetTokenData& VehicleStateData);
+	bool NetSerialize(FArchive& Ar, UPackageMap* Map, bool& bOutSuccess);
+	bool DeltaNetSerialize(FArchive& Ar,  UPackageMap* Map, bool& bOutSuccess);
 
 	/** Interpolate the data in between two states */
 	virtual void InterpolateData(const FNetworkPhysicsData& MinData, const FNetworkPhysicsData& MaxData) override;
diff --git a/Engine/Plugins/Experimental/ChaosModularVehicle/Source/ChaosModularVehicleEngine/Public/ChaosModularVehicle/ModuleInputTokenStore.h b/Engine/Plugins/Experimental/ChaosModularVehicle/Source/ChaosModularVehicleEngine/Public/ChaosModularVehicle/ModuleInputTokenStore.h
--- a/Engine/Plugins/Experimental/ChaosModularVehicle/Source/ChaosModularVehicleEngine/Public/ChaosModularVehicle/ModuleInputTokenStore.h
+++ b/Engine/Plugins/Experimental/ChaosModularVehicle/Source/ChaosModularVehicleEngine/Public/ChaosModularVehicle/ModuleInputTokenStore.h
@@ -21,7 +21,7 @@ struct FModuleInputNetTokenData
 	
 	CHAOSMODULARVEHICLEENGINE_API uint64 GetUniqueKey() const
 	{
-		uint64 HashOfTypes = CityHash64((const char*)Types.GetData(),sizeof(EModuleInputValueType)*Types.Num());
+		uint64 HashOfTypes = CityHash64((const char*)Types.GetData(),sizeof(uint8)*Types.Num());
 		uint64 HashOfDecayValues = CityHash64((const char*)DecayValues.GetData(),sizeof(bool)*DecayValues.Num());
 		return HashOfTypes ^ HashOfDecayValues;
 	}
diff --git a/Engine/Source/Runtime/Experimental/ChaosVehicles/ChaosVehiclesCore/Private/SimModule/ModuleInput.cpp b/Engine/Source/Runtime/Experimental/ChaosVehicles/ChaosVehiclesCore/Private/SimModule/ModuleInput.cpp
--- a/Engine/Source/Runtime/Experimental/ChaosVehicles/ChaosVehiclesCore/Private/SimModule/ModuleInput.cpp
+++ b/Engine/Source/Runtime/Experimental/ChaosVehicles/ChaosVehiclesCore/Private/SimModule/ModuleInput.cpp
@@ -59,7 +59,6 @@ float FModuleInputValue::GetMagnitude() const
 	case EModuleInputValueType::MBoolean:
 	case EModuleInputValueType::MInteger:
 		return ValueInt;
-
 	case EModuleInputValueType::MAxis1D:
 		return Value.X;
 	case EModuleInputValueType::MAxis2D:
@@ -100,8 +99,7 @@ void FModuleInputValue::SetMagnitude(float NewSize)
 	case EModuleInputValueType::MBoolean:
 	case EModuleInputValueType::MInteger:
 		ValueInt = (int32)(NewSize);
-	break;
-
+		break;
 	case EModuleInputValueType::MAxis1D:
 		Value.X = NewSize;
 		break;
@@ -121,6 +119,55 @@ void FModuleInputValue::Serialize(FArchive& Ar, UPackageMap* Map, bool& bOutSucc
 {
 	// #TODO: send only value changes/deltas
 
+	bool bIsNonZero = Ar.IsLoading() ? false : IsQuantizedNonZero();
+	Ar.SerializeBits(&bIsNonZero, 1);
+	Ar.SerializeBits(&ValueType, 3);
+	Ar.SerializeBits(&bApplyInputDecay, 1);
+	
+
+	if (bIsNonZero)
+	{
+		switch (GetValueType())
+		{
+		case EModuleInputValueType::MBoolean:
+			{
+				bool bState = (ValueInt != 0);
+				Ar.SerializeBits(&bState, 1);
+				if (Ar.IsLoading())
+				{
+					ValueInt = bState ? 1 : 0;
+				}
+			}
+			break;
+
+		case EModuleInputValueType::MInteger:
+			{
+				Ar.SerializeIntPacked((uint32&)ValueInt);
+			}
+			break;
+
+		case EModuleInputValueType::MAxis3D:
+			ModularQuantize::SerializeFixedFloat<1, 16>(Value.Z, Ar);
+		case EModuleInputValueType::MAxis2D:
+			ModularQuantize::SerializeFixedFloat<1, 16>(Value.Y, Ar);
+		case EModuleInputValueType::MAxis1D:
+			ModularQuantize::SerializeFixedFloat<1, 16>(Value.X, Ar);
+			break;
+
+		default:
+			checkf(false, TEXT("Unsupported value type for module input value!"));
+			break;
+		}
+	}
+	else
+	{
+		Reset();
+	}
+	bOutSuccess = true;
+}
+
+void FModuleInputValue::NetSerialize(FArchive& Ar, UPackageMap* Map, bool& bOutSuccess)
+{
 	bool bIsNonZero = Ar.IsLoading() ? false : IsQuantizedNonZero();
 	Ar.SerializeBits(&bIsNonZero, 1);
 	if (ChaosModularVehicleCVars::bEnableInputSendingTypeInfo)
@@ -145,10 +192,10 @@ void FModuleInputValue::Serialize(FArchive& Ar, UPackageMap* Map, bool& bOutSucc
 			break;
 
 		case EModuleInputValueType::MInteger:
-		{
-			Ar.SerializeIntPacked((uint32&)ValueInt);
-		}
-		break;
+			{
+				Ar.SerializeIntPacked((uint32&)ValueInt);
+			}
+			break;
 
 		case EModuleInputValueType::MAxis3D:
 			ModularQuantize::SerializeFixedFloat<1, 16>(Value.Z, Ar);
@@ -192,7 +239,7 @@ void FModuleInputValue::DeltaNetSerialize(FArchive& Ar, UPackageMap* Map, bool&
 		return;
 	}
 
-	Serialize(Ar, Map, bOutSuccess);
+	NetSerialize(Ar, Map, bOutSuccess);
 }
 
 void FModuleInputValue::Lerp(const FModuleInputValue& Min, const FModuleInputValue& Max, float Alpha)
diff --git a/Engine/Source/Runtime/Experimental/ChaosVehicles/ChaosVehiclesCore/Public/SimModule/ModuleInput.h b/Engine/Source/Runtime/Experimental/ChaosVehicles/ChaosVehiclesCore/Public/SimModule/ModuleInput.h
--- a/Engine/Source/Runtime/Experimental/ChaosVehicles/ChaosVehiclesCore/Public/SimModule/ModuleInput.h
+++ b/Engine/Source/Runtime/Experimental/ChaosVehicles/ChaosVehiclesCore/Public/SimModule/ModuleInput.h
@@ -339,7 +339,6 @@ public:
 	}
 
 
-
 	template<typename T>
 	inline T Get() const { static_assert(sizeof(T) == 0, "Unsupported conversion for type"); }
 
@@ -368,6 +367,7 @@ public:
 
 	// Serialize values
 	void Serialize(FArchive& Ar, UPackageMap* Map, bool& bOutSuccess);
+	void NetSerialize(FArchive& Ar, UPackageMap* Map, bool& bOutSuccess);
 	void DeltaNetSerialize(FArchive& Ar, UPackageMap* Map, bool& bOutSuccess,const FModuleInputValue& PreviousInputValue);
 
 	void Lerp(const FModuleInputValue& Min, const FModuleInputValue& Max, float Alpha);
