Fix array index check in ChaosSimModuleManagerAsyncCallback.

diff --git a/Engine/Plugins/Experimental/ChaosModularVehicle/Source/ChaosModularVehicleEngine/Private/ChaosModularVehicle/ChaosSimModuleManagerAsyncCallback.cpp b/Engine/Plugins/Experimental/ChaosModularVehicle/Source/ChaosModularVehicleEngine/Private/ChaosModularVehicle/ChaosSimModuleManagerAsyncCallback.cpp
--- a/Engine/Plugins/Experimental/ChaosModularVehicle/Source/ChaosModularVehicleEngine/Private/ChaosModularVehicle/ChaosSimModuleManagerAsyncCallback.cpp
+++ b/Engine/Plugins/Experimental/ChaosModularVehicle/Source/ChaosModularVehicleEngine/Private/ChaosModularVehicle/ChaosSimModuleManagerAsyncCallback.cpp
@@ -411,7 +411,7 @@ bool FNetworkModularVehicleStates::NetSerialize(FArchive& Ar, class UPackageMap*
 					}
 				}
 
-				if (I <= static_cast<uint32>(ModuleData.Num()) && ModuleData[I].IsValid())
+				if (I < static_cast<uint32>(ModuleData.Num()) && ModuleData[I].IsValid())
 				{
 					check(ModuleTypeHash == Chaos::FModuleFactoryRegister::GetModuleHash(ModuleData[I]->GetSimType()));
 
