Fix crash - bug with wrong index being compared, code tidier and easier to understand

diff --git a/Engine/Plugins/Experimental/ChaosModularVehicle/Source/ChaosModularVehicleEngine/Private/ChaosModularVehicle/ChaosSimModuleManagerAsyncCallback.cpp b/Engine/Plugins/Experimental/ChaosModularVehicle/Source/ChaosModularVehicleEngine/Private/ChaosModularVehicle/ChaosSimModuleManagerAsyncCallback.cpp
--- a/Engine/Plugins/Experimental/ChaosModularVehicle/Source/ChaosModularVehicleEngine/Private/ChaosModularVehicle/ChaosSimModuleManagerAsyncCallback.cpp
+++ b/Engine/Plugins/Experimental/ChaosModularVehicle/Source/ChaosModularVehicleEngine/Private/ChaosModularVehicle/ChaosSimModuleManagerAsyncCallback.cpp
@@ -391,39 +391,34 @@ bool FNetworkModularVehicleStates::NetSerialize(FArchive& Ar, class UPackageMap*
 	{
 		Ar.SerializeBits(ModulesBitArray.GetData(), NumNetModules);
 
+		if (NumNetModules != ModuleData.Num())
+		{
+			ModuleData.SetNum(NumNetModules);
+		}
+
 		for (uint32 I = 0; I < NumNetModules; I++)
 		{
-			if (NumNetModules > 0)
-			{
-				uint32 ModuleTypeHash = 0;
-				uint32 SimArrayIndexUnsigned = 0;
+			uint32 ModuleTypeHash = 0;
+			uint32 SimArrayIndexUnsigned = 0;
+
+			Ar << ModuleTypeHash;
+			Ar.SerializeIntPacked(SimArrayIndexUnsigned);
 
-				Ar << ModuleTypeHash;
-				Ar.SerializeIntPacked(SimArrayIndexUnsigned);
+			const int32 SimArrayIndex = static_cast<int32>(SimArrayIndexUnsigned) - 1; // Convert back to signed and adjust
 
-				const int32 SimArrayIndex = static_cast<int32>(SimArrayIndexUnsigned) - 1; // Convert back to signed and adjust
+			if (TSharedPtr<Chaos::FModuleNetData> Data = Chaos::FModuleFactoryRegister::Get().GenerateNetData(ModuleTypeHash, SimArrayIndex))
+			{
+				check(ModuleTypeHash == Chaos::FModuleFactoryRegister::GetModuleHash(Data->GetSimType()));
+				ModuleData[I] = Data;
 
-				if (I >= static_cast<uint32>(ModuleData.Num()))
+				const bool bHasSerializedData = ModulesBitArray.IsBitSet(I);
+				if (bHasSerializedData)
 				{
-					if (TSharedPtr<Chaos::FModuleNetData> Data = Chaos::FModuleFactoryRegister::Get().GenerateNetData(ModuleTypeHash, SimArrayIndex))
-					{
-						ModuleData.Emplace(Data);
-					}
+					ModuleData[I]->Serialize(Ar);
 				}
-
-				if (I < static_cast<uint32>(ModuleData.Num()) && ModuleData[I].IsValid())
+				else
 				{
-					check(ModuleTypeHash == Chaos::FModuleFactoryRegister::GetModuleHash(ModuleData[I]->GetSimType()));
-
-					const bool bHasSerializedData = ModulesBitArray.IsBitSet(I);
-					if (bHasSerializedData)
-					{
-						ModuleData[I]->Serialize(Ar);
-					}
-					else
-					{
-						ModuleData[I]->ApplyDefaultState();
-					}
+					ModuleData[I]->ApplyDefaultState();
 				}
 			}
 		}
