[ModularVehicle] Adding a "CombineValue" method to modular vehicle inputs

Some input schemes apply conflicting values to the same control name with the intent that these values will combine and neutralize, rather than using two separate control names (e.g. "Steer" will receive -1 to turn left and 1 to turn right which combines to 0 if both buttons are pressed at the same time, rather than "SteerLeft" and "SteerRight" both receiving 1). CombineValue allows for this combination. SetValue still remains in case a user wants to explicitly set a specific value.

diff --git a/Engine/Plugins/Experimental/ChaosModularVehicle/Source/ChaosModularVehicleEngine/Private/ChaosModularVehicle/InputProducer.cpp b/Engine/Plugins/Experimental/ChaosModularVehicle/Source/ChaosModularVehicleEngine/Private/ChaosModularVehicle/InputProducer.cpp
--- a/Engine/Plugins/Experimental/ChaosModularVehicle/Source/ChaosModularVehicleEngine/Private/ChaosModularVehicle/InputProducer.cpp
+++ b/Engine/Plugins/Experimental/ChaosModularVehicle/Source/ChaosModularVehicleEngine/Private/ChaosModularVehicle/InputProducer.cpp
@@ -11,13 +11,28 @@ void UVehicleDefaultInputProducer::InitializeContainer(TArray<FModuleInputSetup>
 	MergedInput.Initialize(SetupData, NameMapOut);
 }
 
-void UVehicleDefaultInputProducer::BufferInput(const FInputNameMap& InNameMap, const FName InName, const FModuleInputValue& InValue)
+void UVehicleDefaultInputProducer::BufferInput(const FInputNameMap& InNameMap, const FName InName, const FModuleInputValue& InValue, EModuleInputBufferActionType BufferAction)
 {
 	//UE_LOG(LogTemp, Warning, TEXT("BufferInput - InName %s Val=%s"), *InName.ToString(), *InValue.ToString());
 
-	// inputs are merged here rather than buffered, since they would be merged anyway before use in ProduceInput
-	FInputInterface Inputs(InNameMap, MergedInput, InputQuantizationType);
-	Inputs.MergeValue(InName, InValue);
+	switch(BufferAction)
+	{
+	case EModuleInputBufferActionType::Override:
+		{
+			// inputs are merged here rather than buffered, since they would be merged anyway before use in ProduceInput
+			FInputInterface Inputs(InNameMap, MergedInput, InputQuantizationType);
+			Inputs.MergeValue(InName, InValue);
+			break;
+		}
+	case EModuleInputBufferActionType::Combine:
+		{
+			FInputInterface Inputs(InNameMap, MergedInput, InputQuantizationType);
+			Inputs.CombineValue(InName, InValue);
+			break;
+		}
+	default:
+		break;
+	}
 }
 
 void UVehicleDefaultInputProducer::ProduceInput(int32 PhysicsStep, int32 NumSteps, const FInputNameMap& InNameMap, FModuleInputContainer& InOutContainer)
@@ -62,7 +77,7 @@ void UVehiclePlaybackInputProducer::InitializeContainer(TArray<FModuleInputSetup
 
 }
 
-void UVehiclePlaybackInputProducer::BufferInput(const FInputNameMap& InNameMap, const FName InName, const FModuleInputValue& InValue)
+void UVehiclePlaybackInputProducer::BufferInput(const FInputNameMap& InNameMap, const FName InName, const FModuleInputValue& InValue, EModuleInputBufferActionType BufferAction)
 {
 	// NOP
 }
@@ -98,7 +113,7 @@ void UVehicleRandomInputProducer::InitializeContainer(TArray<FModuleInputSetup>&
 	PlaybackContainer.Initialize(SetupData, NameMapOut);
 }
 
-void UVehicleRandomInputProducer::BufferInput(const FInputNameMap& InNameMap, const FName InName, const FModuleInputValue& InValue)
+void UVehicleRandomInputProducer::BufferInput(const FInputNameMap& InNameMap, const FName InName, const FModuleInputValue& InValue, EModuleInputBufferActionType BufferAction)
 {
 	// NOP
 }
diff --git a/Engine/Plugins/Experimental/ChaosModularVehicle/Source/ChaosModularVehicleEngine/Private/ChaosModularVehicle/ModularVehicleBaseComponent.cpp b/Engine/Plugins/Experimental/ChaosModularVehicle/Source/ChaosModularVehicleEngine/Private/ChaosModularVehicle/ModularVehicleBaseComponent.cpp
--- a/Engine/Plugins/Experimental/ChaosModularVehicle/Source/ChaosModularVehicleEngine/Private/ChaosModularVehicle/ModularVehicleBaseComponent.cpp
+++ b/Engine/Plugins/Experimental/ChaosModularVehicle/Source/ChaosModularVehicleEngine/Private/ChaosModularVehicle/ModularVehicleBaseComponent.cpp
@@ -1521,76 +1521,76 @@ void UModularVehicleBaseComponent::SetInputProducerClass(TSubclassOf<UVehicleInp
 	}
 }
 
-void UModularVehicleBaseComponent::SetInputBool(const FName Name, const bool Value)
+void UModularVehicleBaseComponent::SetInputBool(const FName Name, const bool Value, EModuleInputBufferActionType BufferAction)
 {
 	if (InputProducer)
 	{
-		InputProducer->BufferInput(InputNameMap, Name, Value);
+		InputProducer->BufferInput(InputNameMap, Name, Value, BufferAction);
 	}
 }
 
-void UModularVehicleBaseComponent::SetInputInteger(const FName Name, const int32 Value)
+void UModularVehicleBaseComponent::SetInputInteger(const FName Name, const int32 Value, EModuleInputBufferActionType BufferAction)
 {
 	if (InputProducer)
 	{
-		InputProducer->BufferInput(InputNameMap, Name, Value);
+		InputProducer->BufferInput(InputNameMap, Name, Value, BufferAction);
 	}
 }
 
 
-void UModularVehicleBaseComponent::SetInputAxis1D(const FName Name, const double Value)
+void UModularVehicleBaseComponent::SetInputAxis1D(const FName Name, const double Value, EModuleInputBufferActionType BufferAction)
 {
 	if (InputProducer)
 	{
-		InputProducer->BufferInput(InputNameMap, Name, Value);
+		InputProducer->BufferInput(InputNameMap, Name, Value, BufferAction);
 	}
 }
 
-void UModularVehicleBaseComponent::SetInputAxis2D(const FName Name, const FVector2D Value)
+void UModularVehicleBaseComponent::SetInputAxis2D(const FName Name, const FVector2D Value, EModuleInputBufferActionType BufferAction)
 {
 	if (InputProducer)
 	{
-		InputProducer->BufferInput(InputNameMap, Name, Value);
+		InputProducer->BufferInput(InputNameMap, Name, Value, BufferAction);
 	}
 }
 
-void UModularVehicleBaseComponent::SetInputAxis3D(const FName Name, const FVector Value)
+void UModularVehicleBaseComponent::SetInputAxis3D(const FName Name, const FVector Value, EModuleInputBufferActionType BufferAction)
 {
 	if (InputProducer)
 	{
-		InputProducer->BufferInput(InputNameMap, Name, Value);
+		InputProducer->BufferInput(InputNameMap, Name, Value, BufferAction);
 	}
 }
 
-void UModularVehicleBaseComponent::SetInput(const FName& Name, const bool Value)
+void UModularVehicleBaseComponent::SetInput(const FName& Name, const bool Value, EModuleInputBufferActionType BufferAction)
 {
 	if (InputProducer)
 	{
-		InputProducer->BufferInput(InputNameMap, Name, Value);
+		InputProducer->BufferInput(InputNameMap, Name, Value, BufferAction);
 	}
 }
 
-void UModularVehicleBaseComponent::SetInput(const FName& Name, const double Value)
+void UModularVehicleBaseComponent::SetInput(const FName& Name, const double Value, EModuleInputBufferActionType BufferAction)
 {
 	if (InputProducer)
 	{
-		InputProducer->BufferInput(InputNameMap, Name, Value);
+		InputProducer->BufferInput(InputNameMap, Name, Value, BufferAction);
 	}
 }
 
-void UModularVehicleBaseComponent::SetInput(const FName& Name, const FVector2D& Value)
+void UModularVehicleBaseComponent::SetInput(const FName& Name, const FVector2D& Value, EModuleInputBufferActionType BufferAction)
 {
 	if (InputProducer)
 	{
-		InputProducer->BufferInput(InputNameMap, Name, Value);
+		InputProducer->BufferInput(InputNameMap, Name, Value, BufferAction);
 	}
 }
 
-void UModularVehicleBaseComponent::SetInput(const FName& Name, const FVector& Value)
+void UModularVehicleBaseComponent::SetInput(const FName& Name, const FVector& Value, EModuleInputBufferActionType BufferAction)
 {
 	if (InputProducer)
 	{
-		InputProducer->BufferInput(InputNameMap, Name, Value);
+		InputProducer->BufferInput(InputNameMap, Name, Value, BufferAction);
 	}
 }
 
diff --git a/Engine/Plugins/Experimental/ChaosModularVehicle/Source/ChaosModularVehicleEngine/Public/ChaosModularVehicle/InputProducer.h b/Engine/Plugins/Experimental/ChaosModularVehicle/Source/ChaosModularVehicleEngine/Public/ChaosModularVehicle/InputProducer.h
--- a/Engine/Plugins/Experimental/ChaosModularVehicle/Source/ChaosModularVehicleEngine/Public/ChaosModularVehicle/InputProducer.h
+++ b/Engine/Plugins/Experimental/ChaosModularVehicle/Source/ChaosModularVehicleEngine/Public/ChaosModularVehicle/InputProducer.h
@@ -19,7 +19,7 @@ public:
 	UE_API virtual void InitializeContainer(TArray<FModuleInputSetup>& SetupData, FInputNameMap& NameMapOut, EModuleInputQuantizationType InInputQuantizationType) override;
 
 	/** capture input at game thread frequency */
-	UE_API virtual void BufferInput(const FInputNameMap& InNameMap, const FName InName, const FModuleInputValue& InValue) override;
+	UE_API virtual void BufferInput(const FInputNameMap& InNameMap, const FName InName, const FModuleInputValue& InValue, EModuleInputBufferActionType BufferAction) override;
 
 	/** produce input for PT simulation at PT frequency */
 	UE_API virtual void ProduceInput(int32 PhysicsStep, int32 NumSteps, const FInputNameMap& InNameMap, FModuleInputContainer& InOutContainer) override;
@@ -40,7 +40,7 @@ public:
 	UE_API virtual void InitializeContainer(TArray<FModuleInputSetup>& SetupData, FInputNameMap& NameMapOut, EModuleInputQuantizationType InInputQuantizationType) override;
 
 	/** capture input at game thread frequency */
-	UE_API virtual void BufferInput(const FInputNameMap& InNameMap, const FName InName, const FModuleInputValue& InValue) override;
+	UE_API virtual void BufferInput(const FInputNameMap& InNameMap, const FName InName, const FModuleInputValue& InValue, EModuleInputBufferActionType BufferAction) override;
 
 	/** produce input for PT simulation at PT frequency */
 	UE_API virtual void ProduceInput(int32 PhysicsStep, int32 NumSteps, const FInputNameMap& InNameMap, FModuleInputContainer& InOutContainer) override;
@@ -69,7 +69,7 @@ public:
 	UE_API virtual void InitializeContainer(TArray<FModuleInputSetup>& SetupData, FInputNameMap& NameMapOut, EModuleInputQuantizationType InInputQuantizationType) override;
 
 	/** capture input at game thread frequency */
-	UE_API virtual void BufferInput(const FInputNameMap& InNameMap, const FName InName, const FModuleInputValue& InValue) override;
+	UE_API virtual void BufferInput(const FInputNameMap& InNameMap, const FName InName, const FModuleInputValue& InValue, EModuleInputBufferActionType BufferAction) override;
 
 	/** produce input for PT simulation at PT frequency */
 	UE_API virtual void ProduceInput(int32 PhysicsStep, int32 NumSteps, const FInputNameMap& InNameMap, FModuleInputContainer& InOutContainer) override;
diff --git a/Engine/Plugins/Experimental/ChaosModularVehicle/Source/ChaosModularVehicleEngine/Public/ChaosModularVehicle/ModularVehicleBaseComponent.h b/Engine/Plugins/Experimental/ChaosModularVehicle/Source/ChaosModularVehicleEngine/Public/ChaosModularVehicle/ModularVehicleBaseComponent.h
--- a/Engine/Plugins/Experimental/ChaosModularVehicle/Source/ChaosModularVehicleEngine/Public/ChaosModularVehicle/ModularVehicleBaseComponent.h
+++ b/Engine/Plugins/Experimental/ChaosModularVehicle/Source/ChaosModularVehicleEngine/Public/ChaosModularVehicle/ModularVehicleBaseComponent.h
@@ -215,11 +215,11 @@ public:
 	// 
 	//void SetInput(const FName& Name, const FModuleInputValue& Value);
 	// Control inputs are automatically quantized for networking
-	UE_API void SetInput(const FName& Name, const bool Value);
-	UE_API void SetInput(const FName& Name, const int32 Value);
-	UE_API void SetInput(const FName& Name, const double Value);
-	UE_API void SetInput(const FName& Name, const FVector2D& Value);
-	UE_API void SetInput(const FName& Name, const FVector& Value);
+	UE_API void SetInput(const FName& Name, const bool Value, EModuleInputBufferActionType BufferAction = EModuleInputBufferActionType::Override);
+	UE_API void SetInput(const FName& Name, const int32 Value, EModuleInputBufferActionType BufferAction = EModuleInputBufferActionType::Override);
+	UE_API void SetInput(const FName& Name, const double Value, EModuleInputBufferActionType BufferAction = EModuleInputBufferActionType::Override);
+	UE_API void SetInput(const FName& Name, const FVector2D& Value, EModuleInputBufferActionType BufferAction = EModuleInputBufferActionType::Override);
+	UE_API void SetInput(const FName& Name, const FVector& Value, EModuleInputBufferActionType BufferAction = EModuleInputBufferActionType::Override);
 
 	// The state inputs are not automatically quantized
 	UE_API void SetState(const FName& Name, const bool Value);
@@ -235,19 +235,19 @@ public:
 	TObjectPtr<UVehicleInputProducerBase> GetInputProducer() { return InputProducer; }
 
 	UFUNCTION(BlueprintCallable, Category = "Game|Components|ModularVehicle")
-	UE_API void SetInputBool(const FName Name, const bool Value);
+	UE_API void SetInputBool(const FName Name, const bool Value, EModuleInputBufferActionType BufferAction = EModuleInputBufferActionType::Override);
 
 	UFUNCTION(BlueprintCallable, Category = "Game|Components|ModularVehicle")
-	UE_API void SetInputInteger(const FName Name, const int32 Value);
+	UE_API void SetInputInteger(const FName Name, const int32 Value, EModuleInputBufferActionType BufferAction = EModuleInputBufferActionType::Override);
 
 	UFUNCTION(BlueprintCallable, Category = "Game|Components|ModularVehicle")
-	UE_API void SetInputAxis1D(const FName Name, const double Value);
+	UE_API void SetInputAxis1D(const FName Name, const double Value, EModuleInputBufferActionType BufferAction = EModuleInputBufferActionType::Override);
 
 	UFUNCTION(BlueprintCallable, Category = "Game|Components|ModularVehicle")
-	UE_API void SetInputAxis2D(const FName Name, const FVector2D Value);
+	UE_API void SetInputAxis2D(const FName Name, const FVector2D Value, EModuleInputBufferActionType BufferAction = EModuleInputBufferActionType::Override);
 
 	UFUNCTION(BlueprintCallable, Category = "Game|Components|ModularVehicle")
-	UE_API void SetInputAxis3D(const FName Name, const FVector Value);
+	UE_API void SetInputAxis3D(const FName Name, const FVector Value, EModuleInputBufferActionType BufferAction = EModuleInputBufferActionType::Override);
 
 	/** Set the gear directly */
 	UFUNCTION(BlueprintCallable, Category = "Game|Components|ModularVehicle")
diff --git a/Engine/Source/Runtime/Experimental/ChaosVehicles/ChaosVehiclesCore/Private/SimModule/ModuleInput.cpp b/Engine/Source/Runtime/Experimental/ChaosVehicles/ChaosVehiclesCore/Private/SimModule/ModuleInput.cpp
--- a/Engine/Source/Runtime/Experimental/ChaosVehicles/ChaosVehiclesCore/Private/SimModule/ModuleInput.cpp
+++ b/Engine/Source/Runtime/Experimental/ChaosVehicles/ChaosVehiclesCore/Private/SimModule/ModuleInput.cpp
@@ -392,7 +392,7 @@ void FModuleInputValue::Merge(const FModuleInputValue& From)
 	switch (GetValueType())
 	{
 	case EModuleInputValueType::MBoolean:
-		// ensure we capture edges of digitial inputs by taking largest absolute value
+		// ensure we capture edges of digital inputs by taking largest absolute value
 		if (FMath::Abs(From.ValueInt) >= FMath::Abs(ValueInt))
 		{
 			ValueInt = From.ValueInt;
@@ -415,6 +415,33 @@ void FModuleInputValue::Merge(const FModuleInputValue& From)
 	}
 }
 
+void FModuleInputValue::Combine(const FModuleInputValue& With)
+{
+	switch (GetValueType())
+	{
+	case EModuleInputValueType::MBoolean:
+		// ensure we capture edges of digital inputs by taking largest absolute value
+		if (FMath::Abs(With.ValueInt) >= FMath::Abs(ValueInt))
+		{
+			ValueInt = With.ValueInt;
+		}
+		break;
+
+	case EModuleInputValueType::MInteger:
+		ValueInt += With.ValueInt;
+		break;
+
+	case EModuleInputValueType::MAxis1D:
+	case EModuleInputValueType::MAxis2D:
+	case EModuleInputValueType::MAxis3D:
+		Value += With.Value;
+		break;
+	default:
+		checkf(false, TEXT("Unsupported value type for module input value!"));
+		break;
+	}
+}
+
 void FModuleInputValue::Decay(const float DecayAmount)
 {
 	if (!ShouldApplyInputDecay())
@@ -602,6 +629,15 @@ void FModuleInputContainer::Merge(const FModuleInputContainer& From)
 	}
 }
 
+void FModuleInputContainer::Combine(const FModuleInputContainer& With)
+{
+	int Num = FMath::Min(InputValues.Num(), With.InputValues.Num());
+	for (int I = 0; I < Num; I++)
+	{
+		InputValues[I].Combine(With.InputValues[I]);
+	}
+}
+
 void FModuleInputContainer::Decay(const float DecayAmount)
 {
 	for (int I = 0; I < InputValues.Num(); I++)
@@ -625,6 +661,21 @@ void FInputInterface::SetValue(const FName& InName, const FModuleInputValue& InV
 	}
 }
 
+void FInputInterface::CombineValue(const FName& InName, const FModuleInputValue& InValue)
+{
+	if (ValueContainer.GetNumInputs() > 0)
+	{
+		if (const int* Index = NameMap.Find(InName))
+		{
+			ValueContainer.CombineValueAtIndex(*Index, InValue, InputQuantizationType);
+		}
+		else
+		{
+			UE_LOG(LogModularInput, Warning, TEXT("Trying to set the value of an undefined control input %s"), *InName.ToString());
+		}
+	}
+}
+
 void FInputInterface::MergeValue(const FName& InName, const FModuleInputValue& InValue)
 {
 	if (ValueContainer.GetNumInputs() > 0)
diff --git a/Engine/Source/Runtime/Experimental/ChaosVehicles/ChaosVehiclesCore/Public/SimModule/ModuleInput.h b/Engine/Source/Runtime/Experimental/ChaosVehicles/ChaosVehiclesCore/Public/SimModule/ModuleInput.h
--- a/Engine/Source/Runtime/Experimental/ChaosVehicles/ChaosVehiclesCore/Public/SimModule/ModuleInput.h
+++ b/Engine/Source/Runtime/Experimental/ChaosVehicles/ChaosVehiclesCore/Public/SimModule/ModuleInput.h
@@ -40,6 +40,16 @@ enum class EFunctionType : uint8
 	CustomCurve
 };
 
+UENUM(BlueprintType)
+enum class EModuleInputBufferActionType : uint8
+{
+	// Replaces current value in buffer with new value
+	Override = 0,
+
+	// Adds new value to current value in buffer
+	Combine
+};
+
 namespace ModularQuantize
 {
 	template<int32 MaxValue, uint32 NumBits>
@@ -350,6 +360,8 @@ public:
 
 	UE_API void Merge(const FModuleInputValue& From);
 
+	UE_API void Combine(const FModuleInputValue& With);
+
 	UE_API void Decay(const float DecayAmount);
 
 	// Type sensitive debug stringify
@@ -591,6 +603,11 @@ public:
 		InputValues[Index].Merge(InValue.ReturnQuantized(InInputQuantizationType));
 	}
 
+	void CombineValueAtIndex(int Index, const FModuleInputValue& InValue, EModuleInputQuantizationType InInputQuantizationType)
+	{
+		InputValues[Index].Combine(InValue.ReturnQuantized(InInputQuantizationType));
+	}
+
 	FModuleInputContainer& operator=(const FModuleInputContainer& Other)
 	{
 		if (&Other != this)
@@ -619,6 +636,8 @@ public:
 
 	UE_API void Merge(const FModuleInputContainer& From);
 
+	UE_API void Combine(const FModuleInputContainer& With);
+
 	/**  Decay input during resimulation by DecayAmount which increases over resimulation frames from 0.0 -> 1.0 when the input is being reused */
 	UE_API void Decay(const float DecayAmount);
 
@@ -645,6 +664,7 @@ public:
 	}
 
 	UE_API void SetValue(const FName& InName, const FModuleInputValue& InValue, bool Quantize = true);
+	UE_API void CombineValue(const FName& InName, const FModuleInputValue& InValue);
 	UE_API void MergeValue(const FName& InName, const FModuleInputValue& InValue);
 	UE_API FModuleInputValue GetValue(const FName& InName) const;
 	UE_API EModuleInputValueType GetValueType(const FName& InName) const;
@@ -729,7 +749,7 @@ public:
 	}
 
 	/** capture input at game thread frequency */
-	virtual void BufferInput(const FInputNameMap& InNameMap, const FName InName, const FModuleInputValue& InValue) {}
+	virtual void BufferInput(const FInputNameMap& InNameMap, const FName InName, const FModuleInputValue& InValue, EModuleInputBufferActionType BufferAction = EModuleInputBufferActionType::Override) {}
 
 	/** produce input for PT simulation at PT frequency */
 	virtual void ProduceInput(int32 PhysicsStep, int32 NumSteps, const FInputNameMap& InNameMap, FModuleInputContainer& InOutContainer) {}
