Add Module Initialized Event

diff --git a/Engine/Plugins/Experimental/ChaosModularVehicle/Source/ChaosModularVehicleEngine/Private/ChaosModularVehicle/ModularVehicleBaseComponent.cpp b/Engine/Plugins/Experimental/ChaosModularVehicle/Source/ChaosModularVehicleEngine/Private/ChaosModularVehicle/ModularVehicleBaseComponent.cpp
--- a/Engine/Plugins/Experimental/ChaosModularVehicle/Source/ChaosModularVehicleEngine/Private/ChaosModularVehicle/ModularVehicleBaseComponent.cpp
+++ b/Engine/Plugins/Experimental/ChaosModularVehicle/Source/ChaosModularVehicleEngine/Private/ChaosModularVehicle/ModularVehicleBaseComponent.cpp
@@ -171,6 +171,9 @@ void UModularVehicleBaseComponent::OnCreatePhysicsState()
 		ClusterUnionComponent->OnComponentAddedNativeEvent.AddUObject(this, &UModularVehicleBaseComponent::AddComponentToSimulation);
 		ClusterUnionComponent->OnComponentRemovedNativeEvent.AddUObject(this, &UModularVehicleBaseComponent::RemoveComponentFromSimulation);
 
+		OnModuleAddedNativeEvent.AddUObject(this, &UModularVehicleBaseComponent::OnModuleInitialized);
+		OnModuleRemovedNativeEvent.AddUObject(this, &UModularVehicleBaseComponent::OnModuleRemoved);
+
 		// #TODO: this doesn't appear to be working
 		if (bKeepVehicleAwake && ClusterUnionComponent->GetPhysicsProxy())
 		{
@@ -184,6 +187,8 @@ void UModularVehicleBaseComponent::OnCreatePhysicsState()
 	// convert the game thread side UVehicleSimBaseComponent(s) to physics thread simulation ISimulationModuleBase(s)
 	CreateVehicleSim();
 
+	SetupInputConfiguration();
+
 	if (bUsingNetworkPhysicsPrediction)
 	{
 		if (NetworkPhysicsComponent)
@@ -217,6 +222,9 @@ void UModularVehicleBaseComponent::OnDestroyPhysicsState()
 		ClusterUnionComponent->OnComponentRemovedNativeEvent.RemoveAll(this);
 	}
 
+	OnModuleAddedNativeEvent.RemoveAll(this);
+	OnModuleRemovedNativeEvent.RemoveAll(this);
+
 	DestroyVehicleSim();
 
 	if (bUsingNetworkPhysicsPrediction && NetworkPhysicsComponent)
@@ -375,6 +383,8 @@ void UModularVehicleBaseComponent::CreateAssociatedSimComponents(USceneComponent
 		ComponentData.Guid = GenerateNewGuid();
 		NewModule->SetGuid(ComponentData.Guid);
 
+		PhysicsGuidToComponent.Add(ComponentData.Guid, AttachedComponent);
+
 		NewModule->SetInitialParticleTransform(InitialTransform);
 		NewModule->SetComponentTransform(ComponentTransform);
 
@@ -492,6 +502,8 @@ void UModularVehicleBaseComponent::CreateIndependentSimComponents(USceneComponen
 		ComponentData.Guid = GenerateNewGuid();
 		NewModule->SetGuid(ComponentData.Guid);
 
+		PhysicsGuidToComponent.Add(ComponentData.Guid, AttachedComponent);
+
 		NewModule->SetInitialParticleTransform(InitialTransform);
 		NewModule->SetComponentTransform(ComponentTransform);
 
@@ -647,6 +659,7 @@ void UModularVehicleBaseComponent::BeginPlay()
 {
 	Super::BeginPlay();
 
+	// this only happens on the server - the client is setup from OnComponentAddedNativeEvent callback
 	const bool bHasAuthority = GetOwner()->HasAuthority();
 	if (bHasAuthority)
 	{
@@ -660,8 +673,10 @@ void UModularVehicleBaseComponent::BeginPlay()
 		}
 	}
 
-	//AddSimulationComponentsFromOwnedActor();
-	
+}
+
+void UModularVehicleBaseComponent::SetupInputConfiguration()
+{
 	// control input setup - unfortunately can't do this in OnCreatePhysics since RootComponent->GetChildrenComponents will not work 
 	// at that time and AssimilateComponentInputs will not find any controls in the component hierarchy
 	TArray<FModuleInputSetup> CombinedInputConfiguration;
@@ -1082,6 +1097,14 @@ void UModularVehicleBaseComponent::ActionTreeUpdates(Chaos::FSimTreeUpdates* Nex
 	{
 		if (Chaos::FPBDRigidsSolver* Solver = Proxy->GetSolver<Chaos::FPBDRigidsSolver>())
 		{
+			for (const Chaos::FPendingModuleAdds& AddedModule : NextTreeUpdates->GetNewModules())
+			{
+				if (AddedModule.NewSimModule)
+				{
+					BroadcastModuleAddedEvent(AddedModule.NewSimModule->GetSimType(), AddedModule.NewSimModule->GetGuid(), AddedModule.NewSimModule->GetTreeIndex());
+				}
+			}
+
 			Solver->EnqueueCommandImmediate([Proxy, this, NextTreeUpdates = *NextTreeUpdates]() mutable
 			{
 				if (IsValidChecked(this) && bPhysicsStateCreated && VehicleSimulationPT)
@@ -1155,6 +1178,43 @@ int32 UModularVehicleBaseComponent::FindParentsLastSimComponent(const USceneComp
 	return INDEX_NONE;
 }
 
+
+void UModularVehicleBaseComponent::OnModuleInitialized(const FName& SimType, int Guid, int TreeIndex)
+{
+	if (const TWeakObjectPtr<USceneComponent>* Component = PhysicsGuidToComponent.Find(Guid))
+	{
+		if (Component->IsValid())
+		{
+			if (UVehicleSimBaseSceneComponent* BaseSceneComponent = Cast<UVehicleSimBaseSceneComponent>(Component->Get()))
+			{
+				BaseSceneComponent->OnAdded();
+			}
+			else if (UVehicleSimBaseComponent* BaseComponent = Cast<UVehicleSimBaseComponent>(Component->Get()))
+			{
+				BaseComponent->OnAdded();
+			}
+		}
+	}
+}
+
+void UModularVehicleBaseComponent::OnModuleRemoved(const FName& SimType, int Guid, int TreeIndex)
+{
+	if (const TWeakObjectPtr<USceneComponent>* Component = PhysicsGuidToComponent.Find(Guid))
+	{
+		if (Component->IsValid())
+		{
+			if (UVehicleSimBaseSceneComponent* BaseSceneComponent = Cast<UVehicleSimBaseSceneComponent>(Component->Get()))
+			{
+				BaseSceneComponent->OnRemoved();
+			}
+			else if (UVehicleSimBaseComponent* BaseComponent = Cast<UVehicleSimBaseComponent>(Component->Get()))
+			{
+				BaseComponent->OnRemoved();
+			}
+		}
+	}
+}
+
 void UModularVehicleBaseComponent::AddComponentToSimulation(UPrimitiveComponent* InComponent, const TArray<FClusterUnionBoneData>& BonesData, const TArray<FClusterUnionBoneData>& RemovedBoneIDs, bool bIsNew)
 {
 	AddComponentToSimulationImpl(InComponent, BonesData, RemovedBoneIDs, bIsNew);
@@ -1251,6 +1311,9 @@ void UModularVehicleBaseComponent::RemoveComponentFromSimulationImpl(UPrimitiveC
 					{
 						if (SimModule->GetGuid() == TreeUpdate.Guid)
 						{
+							BroadcastModuleRemovedEvent(SimModule->GetSimType(), SimModule->GetGuid(), SimModule->GetTreeIndex());
+							PhysicsGuidToComponent.Remove(SimModule->GetGuid());
+
 							SimModule->SetAnimationEnabled(false);
 							SimModule->SetStateFlags(Chaos::eSimModuleState::Disabled);
 							SimModule->OnTermination_External();
@@ -1379,6 +1442,7 @@ void UModularVehicleBaseComponent::DestroyVehicleSim()
 				PVehicleOutput.Reset(nullptr);
 			}
 
+			// This is a managed pointer, so do we need to do this?
 			if (VehicleSimulationPT)
 			{
 				VehicleSimulationPT->Terminate();
@@ -1639,6 +1703,30 @@ bool UModularVehicleBaseComponent::IsReversing()
 	return (GetCurrentGear() < 0);
 }
 
+void UModularVehicleBaseComponent::BroadcastModuleAddedEvent(const FName& SimType, int Guid, int TreeIndex)
+{
+	if (OnModuleAddedEvent.IsBound())
+	{
+		OnModuleAddedEvent.Broadcast(SimType, Guid, TreeIndex);
+	}
+	if (OnModuleAddedNativeEvent.IsBound())
+	{
+		OnModuleAddedNativeEvent.Broadcast(SimType, Guid, TreeIndex);
+	}
+}
+
+void UModularVehicleBaseComponent::BroadcastModuleRemovedEvent(const FName& SimType, int Guid, int TreeIndex)
+{
+	if (OnModuleRemovedEvent.IsBound())
+	{
+		OnModuleRemovedEvent.Broadcast(SimType, Guid, TreeIndex);
+	}
+	if (OnModuleRemovedNativeEvent.IsBound())
+	{
+		OnModuleRemovedNativeEvent.Broadcast(SimType, Guid, TreeIndex);
+	}
+}
+
 
 //-=====================================================
 // Networking Replication
diff --git a/Engine/Plugins/Experimental/ChaosModularVehicle/Source/ChaosModularVehicleEngine/Public/ChaosModularVehicle/ModularVehicleBaseComponent.h b/Engine/Plugins/Experimental/ChaosModularVehicle/Source/ChaosModularVehicleEngine/Public/ChaosModularVehicle/ModularVehicleBaseComponent.h
--- a/Engine/Plugins/Experimental/ChaosModularVehicle/Source/ChaosModularVehicleEngine/Public/ChaosModularVehicle/ModularVehicleBaseComponent.h
+++ b/Engine/Plugins/Experimental/ChaosModularVehicle/Source/ChaosModularVehicleEngine/Public/ChaosModularVehicle/ModularVehicleBaseComponent.h
@@ -168,6 +168,22 @@ public:
 	}
 
 	UE_API FORCEINLINE const FTransform& GetComponentTransform() const;
+	
+	// native (fast, low overhead) versions 
+	DECLARE_MULTICAST_DELEGATE_ThreeParams(FOnModuleAddedNative, const FName&, int, int);
+	FOnModuleAddedNative OnModuleAddedNativeEvent;
+
+	DECLARE_MULTICAST_DELEGATE_ThreeParams(FOnModuleRemovedNative, const FName&, int, int);
+	FOnModuleRemovedNative OnModuleRemovedNativeEvent;
+
+	// standard events
+	DECLARE_DYNAMIC_MULTICAST_DELEGATE_ThreeParams(FOnModuleAdded, const FName&, SimType, int, Guid, int, TreeIndex);
+	UPROPERTY(BlueprintAssignable, Category = "Events")
+	FOnModuleAdded OnModuleAddedEvent;
+
+	DECLARE_DYNAMIC_MULTICAST_DELEGATE_ThreeParams(FOnModuleRemoved, const FName&, SimType, int, Guid, int, TreeIndex);
+	UPROPERTY(BlueprintAssignable, Category = "Events")
+	FOnModuleRemoved OnModuleRemovedEvent;
 
 	/** Use to naturally decelerate linear velocity of objects */
 	UPROPERTY(EditAnywhere, Category = "Game|Components|ModularVehicle")
@@ -211,6 +227,12 @@ public:
 	UPROPERTY(EditAnywhere, Category = "Game|Components|ModularVehicle")
 	EModuleInputQuantizationType InputQuantizationType = EModuleInputQuantizationType::Default_16Bits;
 	
+	UFUNCTION()
+	UE_API void OnModuleInitialized(const FName& SimType, int Guid, int TreeIndex);
+
+	UFUNCTION()
+	UE_API void OnModuleRemoved(const FName& SimType, int Guid, int TreeIndex);
+
 	// CONTROLS
 	// 
 	//void SetInput(const FName& Name, const FModuleInputValue& Value);
@@ -283,6 +305,9 @@ public:
 	/*** Map simulation component to our vehicle setup data */
 	TMap<TObjectKey<USceneComponent>, FVehicleComponentData> ComponentToPhysicsObjects;
 
+	/*** Map physics sim GUID to actor component */
+	TMap<int, TWeakObjectPtr<USceneComponent>> PhysicsGuidToComponent;
+
 	UClusterUnionComponent* ClusterUnionComponent;
 
 	/** Set all channels to the specified response - for wheel raycasts */
@@ -339,6 +364,9 @@ protected:
 	UE_API int32 FindComponentAddOrder(USceneComponent* InComponent);
 	UE_API bool FindAndRemoveNextPendingUpdate(int32 NextIndex, Chaos::FSimTreeUpdates* OutData);
 
+	UE_API void BroadcastModuleAddedEvent(const FName& SimType, int Guid, int TreeIndex);
+	UE_API void BroadcastModuleRemovedEvent(const FName& SimType, int Guid, int TreeIndex);
+
 	// replicated state of vehicle 
 	UPROPERTY(Transient, Replicated)
 	FModularReplicatedState ReplicatedState;
@@ -413,6 +441,8 @@ public:
 
 private:
 
+	void SetupInputConfiguration();
+
 	int NextTransformIndex = 0; // is there a better way, getting from size of map/array?
 	UPrimitiveComponent* MyComponent = nullptr;
 
diff --git a/Engine/Plugins/Experimental/ChaosModularVehicle/Source/ChaosModularVehicleEngine/Public/ChaosModularVehicle/VehicleSimBaseComponent.h b/Engine/Plugins/Experimental/ChaosModularVehicle/Source/ChaosModularVehicleEngine/Public/ChaosModularVehicle/VehicleSimBaseComponent.h
--- a/Engine/Plugins/Experimental/ChaosModularVehicle/Source/ChaosModularVehicleEngine/Public/ChaosModularVehicle/VehicleSimBaseComponent.h
+++ b/Engine/Plugins/Experimental/ChaosModularVehicle/Source/ChaosModularVehicleEngine/Public/ChaosModularVehicle/VehicleSimBaseComponent.h
@@ -62,6 +62,8 @@ public:
 	virtual int32 GetAnimationSetupIndex() const { return INDEX_NONE; }
 	virtual void SetTreeIndex(const int32 NewValue) {}
 	virtual int32 GetTreeIndex() const { return INDEX_NONE; }
+	virtual void OnAdded() {}
+	virtual void OnRemoved() {}
 };
 
 /** This if for sim components that need scene component properties along with rendering and collision. */
