Modular vehicle native events promoted to Verse Events

#okforversepublic

diff --git a/Engine/Plugins/Experimental/ChaosModularVehicle/Source/ChaosModularVehicleEngine/Private/ChaosModularVehicle/ModularVehicleBaseComponent.cpp b/Engine/Plugins/Experimental/ChaosModularVehicle/Source/ChaosModularVehicleEngine/Private/ChaosModularVehicle/ModularVehicleBaseComponent.cpp
--- a/Engine/Plugins/Experimental/ChaosModularVehicle/Source/ChaosModularVehicleEngine/Private/ChaosModularVehicle/ModularVehicleBaseComponent.cpp
+++ b/Engine/Plugins/Experimental/ChaosModularVehicle/Source/ChaosModularVehicleEngine/Private/ChaosModularVehicle/ModularVehicleBaseComponent.cpp
@@ -171,8 +171,6 @@ void UModularVehicleBaseComponent::OnCreatePhysicsState()
 		ClusterUnionComponent->OnComponentAddedNativeEvent.AddUObject(this, &UModularVehicleBaseComponent::AddComponentToSimulation);
 		ClusterUnionComponent->OnComponentRemovedNativeEvent.AddUObject(this, &UModularVehicleBaseComponent::RemoveComponentFromSimulation);
 
-		OnModuleAddedNativeEvent.AddUObject(this, &UModularVehicleBaseComponent::OnModuleInitialized);
-		OnModuleRemovedNativeEvent.AddUObject(this, &UModularVehicleBaseComponent::OnModuleRemoved);
 
 		// #TODO: this doesn't appear to be working
 		if (bKeepVehicleAwake && ClusterUnionComponent->GetPhysicsProxy())
@@ -184,6 +182,9 @@ void UModularVehicleBaseComponent::OnCreatePhysicsState()
 		}
 	}
 
+	OnModuleAddedNativeEvent.AddUObject(this, &UModularVehicleBaseComponent::OnModuleInitialized);
+	OnModuleRemovedNativeEvent.AddUObject(this, &UModularVehicleBaseComponent::OnModuleRemoved);
+
 	// convert the game thread side UVehicleSimBaseComponent(s) to physics thread simulation ISimulationModuleBase(s)
 	CreateVehicleSim();
 
@@ -1093,6 +1094,11 @@ const FTransform& UModularVehicleBaseComponent::GetComponentTransform() const
 
 void UModularVehicleBaseComponent::ActionTreeUpdates(Chaos::FSimTreeUpdates* NextTreeUpdates)
 {
+	if (NextTreeUpdates == nullptr)
+	{
+		return;
+	}
+
 	if (IPhysicsProxyBase* Proxy = GetPhysicsProxy())
 	{
 		if (Chaos::FPBDRigidsSolver* Solver = Proxy->GetSolver<Chaos::FPBDRigidsSolver>())
diff --git a/Engine/Plugins/Experimental/ChaosModularVehicle/Source/ChaosModularVehicleEngine/Public/ChaosModularVehicle/ModularVehicleBaseComponent.h b/Engine/Plugins/Experimental/ChaosModularVehicle/Source/ChaosModularVehicleEngine/Public/ChaosModularVehicle/ModularVehicleBaseComponent.h
--- a/Engine/Plugins/Experimental/ChaosModularVehicle/Source/ChaosModularVehicleEngine/Public/ChaosModularVehicle/ModularVehicleBaseComponent.h
+++ b/Engine/Plugins/Experimental/ChaosModularVehicle/Source/ChaosModularVehicleEngine/Public/ChaosModularVehicle/ModularVehicleBaseComponent.h
@@ -228,10 +228,10 @@ public:
 	EModuleInputQuantizationType InputQuantizationType = EModuleInputQuantizationType::Default_16Bits;
 	
 	UFUNCTION()
-	UE_API void OnModuleInitialized(const FName& SimType, int Guid, int TreeIndex);
+	UE_API virtual void OnModuleInitialized(const FName& SimType, int Guid, int TreeIndex);
 
 	UFUNCTION()
-	UE_API void OnModuleRemoved(const FName& SimType, int Guid, int TreeIndex);
+	UE_API virtual void OnModuleRemoved(const FName& SimType, int Guid, int TreeIndex);
 
 	// CONTROLS
 	// 
