[Chaos Modular Vehicles]
Fix case in SimModuleTree::Simulate where SimulateNode was still called in case of ManualOverride mode.
Add a few properties to FVehicleState.
Set up a VehicleSimulationCallback in FModularVehicleSimulation so that the user can call SimulateNode manually.
Make SimulateNode public.
Fix module input setup not happening properly.

diff --git a/Engine/Plugins/Experimental/ChaosModularVehicle/Source/ChaosModularVehicleEngine/Private/ChaosModularVehicle/ModularVehicleBaseComponent.cpp b/Engine/Plugins/Experimental/ChaosModularVehicle/Source/ChaosModularVehicleEngine/Private/ChaosModularVehicle/ModularVehicleBaseComponent.cpp
--- a/Engine/Plugins/Experimental/ChaosModularVehicle/Source/ChaosModularVehicleEngine/Private/ChaosModularVehicle/ModularVehicleBaseComponent.cpp
+++ b/Engine/Plugins/Experimental/ChaosModularVehicle/Source/ChaosModularVehicleEngine/Private/ChaosModularVehicle/ModularVehicleBaseComponent.cpp
@@ -1582,15 +1582,19 @@ void UModularVehicleBaseComponent::AddSimulationComponentsFromOwnedActor()
 	}
 }
 
-void UModularVehicleBaseComponent::SetInputProducerClass(TSubclassOf<UVehicleInputProducerBase> InInputProducerClass)
+void UModularVehicleBaseComponent::SetInputProducerClass(TSubclassOf<UVehicleInputProducerBase> InInputProducerClass, bool bForceNewInstance)
 {
 	InputProducerClass = InInputProducerClass;
-	if (!InputProducer)
+
+	if (bForceNewInstance)
 	{
-		InputProducer = NewObject<UVehicleInputProducerBase>(this, InputProducerClass);
+		InputProducer = nullptr;
 	}
+	
+	SetupInputConfiguration();
 }
 
+
 void UModularVehicleBaseComponent::SetInputBool(const FName Name, const bool Value, EModuleInputBufferActionType BufferAction)
 {
 	if (InputProducer)
diff --git a/Engine/Plugins/Experimental/ChaosModularVehicle/Source/ChaosModularVehicleEngine/Private/ChaosModularVehicle/ModularVehicleSimulationCU.cpp b/Engine/Plugins/Experimental/ChaosModularVehicle/Source/ChaosModularVehicleEngine/Private/ChaosModularVehicle/ModularVehicleSimulationCU.cpp
--- a/Engine/Plugins/Experimental/ChaosModularVehicle/Source/ChaosModularVehicleEngine/Private/ChaosModularVehicle/ModularVehicleSimulationCU.cpp
+++ b/Engine/Plugins/Experimental/ChaosModularVehicle/Source/ChaosModularVehicleEngine/Private/ChaosModularVehicle/ModularVehicleSimulationCU.cpp
@@ -128,6 +128,10 @@ void FModularVehicleSimulation::SimulateModuleTree(UWorld* InWorld, float DeltaS
 
 		// run the dynamics simulation, engine, suspension, wheels, aerofoils etc.
 		SimModuleTree->Simulate(DeltaSeconds, SimInputData, Proxy, RootParticle);
+		if(SimModuleTree->GetSimTreeProcessingOrder() == ESimTreeProcessingOrder::ManualOverride)
+		{
+			VehicleSimulationCallback.Broadcast(DeltaSeconds, SimInputData, Proxy, RootParticle, SimModuleTree.Get());
+		}
 	}
 
 }
diff --git a/Engine/Plugins/Experimental/ChaosModularVehicle/Source/ChaosModularVehicleEngine/Public/ChaosModularVehicle/ModularVehicleBaseComponent.h b/Engine/Plugins/Experimental/ChaosModularVehicle/Source/ChaosModularVehicleEngine/Public/ChaosModularVehicle/ModularVehicleBaseComponent.h
--- a/Engine/Plugins/Experimental/ChaosModularVehicle/Source/ChaosModularVehicleEngine/Public/ChaosModularVehicle/ModularVehicleBaseComponent.h
+++ b/Engine/Plugins/Experimental/ChaosModularVehicle/Source/ChaosModularVehicleEngine/Public/ChaosModularVehicle/ModularVehicleBaseComponent.h
@@ -250,9 +250,11 @@ public:
 	UE_API void SetState(const FName& Name, const FVector2D& Value);
 	UE_API void SetState(const FName& Name, const FVector& Value);
 
-	// Sets the input producer class and creates an input producer if one doesn't exist.
+	UE_API void SetupInputConfiguration();
+
+	// Sets the input producer class and creates an input producer if one doesn't exist, or if forced to create an instance with the new class.
 	UFUNCTION(BlueprintCallable, Category = "Game|Components|ModularVehicle")
-	UE_API void SetInputProducerClass(TSubclassOf<UVehicleInputProducerBase> InInputProducerClass);
+	UE_API void SetInputProducerClass(TSubclassOf<UVehicleInputProducerBase> InInputProducerClass, bool bForceNewInstance = false);
 
 	TObjectPtr<UVehicleInputProducerBase> GetInputProducer() { return InputProducer; }
 
@@ -441,8 +443,6 @@ public:
 
 private:
 
-	void SetupInputConfiguration();
-
 	int NextTransformIndex = 0; // is there a better way, getting from size of map/array?
 	UPrimitiveComponent* MyComponent = nullptr;
 
diff --git a/Engine/Plugins/Experimental/ChaosModularVehicle/Source/ChaosModularVehicleEngine/Public/ChaosModularVehicle/ModularVehicleSimulationCU.h b/Engine/Plugins/Experimental/ChaosModularVehicle/Source/ChaosModularVehicleEngine/Public/ChaosModularVehicle/ModularVehicleSimulationCU.h
--- a/Engine/Plugins/Experimental/ChaosModularVehicle/Source/ChaosModularVehicleEngine/Public/ChaosModularVehicle/ModularVehicleSimulationCU.h
+++ b/Engine/Plugins/Experimental/ChaosModularVehicle/Source/ChaosModularVehicleEngine/Public/ChaosModularVehicle/ModularVehicleSimulationCU.h
@@ -3,6 +3,7 @@
 #pragma once
 
 #include "ChaosModularVehicle/ChaosSimModuleManagerAsyncCallback.h"
+#include "Delegates/Delegate.h"
 #include "SimModule/SimModuleTree.h"
 #include "SimModule/SimModulesInclude.h"
 #include "ModularVehicleBuilder.h"
@@ -36,6 +37,8 @@ namespace Chaos
 {
 	class FClusterUnionPhysicsProxy;
 	class FSingleParticlePhysicsProxy;
+	DECLARE_MULTICAST_DELEGATE_FiveParams(FModularVehicleSimulationCallback, float, FAllInputs&, IPhysicsProxyBase*, Chaos::FPBDRigidParticleHandle*, Chaos::FSimModuleTree*);
+	
 }
 
 
@@ -128,6 +131,7 @@ public:
 
 	Chaos::FPBDRigidParticleHandle* RootParticle = nullptr; // cached root particle
 
+	Chaos::FModularVehicleSimulationCallback VehicleSimulationCallback;
 };
 
 #undef UE_API
diff --git a/Engine/Source/Runtime/Experimental/ChaosVehicles/ChaosVehiclesCore/Private/SimModule/SimModuleTree.cpp b/Engine/Source/Runtime/Experimental/ChaosVehicles/ChaosVehiclesCore/Private/SimModule/SimModuleTree.cpp
--- a/Engine/Source/Runtime/Experimental/ChaosVehicles/ChaosVehiclesCore/Private/SimModule/SimModuleTree.cpp
+++ b/Engine/Source/Runtime/Experimental/ChaosVehicles/ChaosVehiclesCore/Private/SimModule/SimModuleTree.cpp
@@ -241,7 +241,7 @@ void FSimModuleTree::Simulate(float DeltaTime, FAllInputs& Inputs, IPhysicsProxy
 		{
 			SimulateNodeBFS(DeltaTime, Inputs, RootNodes, PhysicsProxy, RootParticle);
 		}
-		else
+		else if(SimTreeProcessingOrder != ESimTreeProcessingOrder::ManualOverride)
 		{
 			for (int RootIndex : RootNodes)
 			{
@@ -514,11 +514,14 @@ void FSimModuleTree::UpdateVehicleState(Chaos::FPBDRigidParticleHandle* RootPart
 	{
 		const FTransform BodyTransform(RootParticle->GetR(), RootParticle->GetX());
 
+		VehicleState.Position =  RootParticle->GetX();
+		VehicleState.Rotation = RootParticle->GetR();
 		VehicleState.ForwardDir = BodyTransform.GetUnitAxis(EAxis::X);
 		VehicleState.UpDir = BodyTransform.GetUnitAxis(EAxis::Z);
 		VehicleState.RightDir = BodyTransform.GetUnitAxis(EAxis::Y);
 		VehicleState.ForwardSpeedKmh = Chaos::CmSToKmH(FVector::DotProduct(RootParticle->GetV(), VehicleState.ForwardDir));
 		VehicleState.AngularVelocityRad = RootParticle->GetW();
+		VehicleState.LinearVelocity = RootParticle->GetV();
 	}
 
 }
diff --git a/Engine/Source/Runtime/Experimental/ChaosVehicles/ChaosVehiclesCore/Public/SimModule/SimModuleTree.h b/Engine/Source/Runtime/Experimental/ChaosVehicles/ChaosVehiclesCore/Public/SimModule/SimModuleTree.h
--- a/Engine/Source/Runtime/Experimental/ChaosVehicles/ChaosVehiclesCore/Public/SimModule/SimModuleTree.h
+++ b/Engine/Source/Runtime/Experimental/ChaosVehicles/ChaosVehiclesCore/Public/SimModule/SimModuleTree.h
@@ -82,10 +82,13 @@ namespace Chaos
 	struct FVehicleState
 	{
 		float ForwardSpeedKmh;
+		FVector Position;
+		FQuat Rotation;
 		FVector ForwardDir;
 		FVector RightDir;
 		FVector UpDir;
 		FVector AngularVelocityRad;
+		FVector LinearVelocity;
 	};
 
 	class FSimModuleTree
@@ -139,6 +142,7 @@ namespace Chaos
 		int NumActiveNodes() const { return SimulationModuleTree.Num() - FreeList.Num(); }
 		UE_API void GetRootNodes(TArray<int>& RootNodesOut);
 		int GetNumNodes() const { return SimulationModuleTree.Num(); }
+		ESimTreeProcessingOrder GetSimTreeProcessingOrder() const { return SimTreeProcessingOrder; }
 
 		UE_API int AddRoot(ISimulationModuleBase* SimModule);
 
@@ -214,8 +218,8 @@ namespace Chaos
 			return SimBlackboard.Get();
 		}
 
-	protected:
 		UE_API void SimulateNode(float DeltaTime, FAllInputs& Inputs, int NodeIdx, IPhysicsProxyBase* PhysicsProxy, Chaos::FPBDRigidParticleHandle* ParticleHandle);
+	protected:
 
 		UE_API void OnContactModificationInternal(int NodeIndex, FCollisionContactModifier& Modifier, IPhysicsProxyBase* PhysicsProxy);
 
