Possible fix for rare crash

diff --git a/Engine/Plugins/Experimental/ChaosModularVehicle/Source/ChaosModularVehicleEngine/Private/ChaosModularVehicle/ModularVehicleBaseComponent.cpp b/Engine/Plugins/Experimental/ChaosModularVehicle/Source/ChaosModularVehicleEngine/Private/ChaosModularVehicle/ModularVehicleBaseComponent.cpp
--- a/Engine/Plugins/Experimental/ChaosModularVehicle/Source/ChaosModularVehicleEngine/Private/ChaosModularVehicle/ModularVehicleBaseComponent.cpp
+++ b/Engine/Plugins/Experimental/ChaosModularVehicle/Source/ChaosModularVehicleEngine/Private/ChaosModularVehicle/ModularVehicleBaseComponent.cpp
@@ -1332,17 +1332,22 @@ void UModularVehicleBaseComponent::RemoveComponentFromSimulationImpl(UPrimitiveC
 
 		Chaos::FClusterUnionPhysicsProxy* Proxy = ClusterUnionComponent->GetPhysicsProxy();
 		Chaos::FPBDRigidsSolver* Solver = Proxy->GetSolver<Chaos::FPBDRigidsSolver>();
-		Solver->EnqueueCommandImmediate([Proxy, this, LatestTreeUpdates = LatestTreeUpdates]() mutable
-			{
-				if (IsValidChecked(this) && bPhysicsStateCreated && VehicleSimulationPT)
+		{
+			Solver->EnqueueCommandImmediate([Proxy, this, LatestTreeUpdates = LatestTreeUpdates]() mutable
 				{
-					TUniquePtr<Chaos::FSimModuleTree>& SimModuleTree = VehicleSimulationPT->AccessSimComponentTree();
-					if(SimModuleTree.IsValid())
+					if (Chaos::FPhysicsSolver* Solver = Proxy->GetSolver<Chaos::FPhysicsSolver>())
 					{
-						SimModuleTree->AppendTreeUpdates(LatestTreeUpdates);
+						if (!Solver->IsShuttingDown() && IsValidChecked(this) && bPhysicsStateCreated && VehicleSimulationPT)
+						{
+							TUniquePtr<Chaos::FSimModuleTree>& SimModuleTree = VehicleSimulationPT->AccessSimComponentTree();
+							if (SimModuleTree.IsValid())
+							{
+								SimModuleTree->AppendTreeUpdates(LatestTreeUpdates);
+							}
+						}
 					}
-				}
-			});
+				});
+		}
 
 		NextTransformIndex--;
 	}
