Modular Vehicle : remove unused code relating to in built input modifier, this behaviour should come from the enhanced input system

diff --git a/Engine/Plugins/Experimental/ChaosModularVehicle/Source/ChaosModularVehicleEngine/Private/ChaosModularVehicle/ModularVehicleBaseComponent.cpp b/Engine/Plugins/Experimental/ChaosModularVehicle/Source/ChaosModularVehicleEngine/Private/ChaosModularVehicle/ModularVehicleBaseComponent.cpp
--- a/Engine/Plugins/Experimental/ChaosModularVehicle/Source/ChaosModularVehicleEngine/Private/ChaosModularVehicle/ModularVehicleBaseComponent.cpp
+++ b/Engine/Plugins/Experimental/ChaosModularVehicle/Source/ChaosModularVehicleEngine/Private/ChaosModularVehicle/ModularVehicleBaseComponent.cpp
@@ -128,39 +128,6 @@ void UModularVehicleBaseComponent::ProduceInput(int32 PhysicsStep, int32 NumStep
 	}
 }
 
-//void UModularVehicleBaseComponent::GenerateInputModifiers(const TArray<FModuleInputSetup>& CombinedInputConfiguration)
-//{
-//	for (const FModuleInputSetup& InputSetup : CombinedInputConfiguration)
-//	{
-//		if (InputSetup.InputModifierClass != nullptr)
-//		{
-//			UDefaultModularVehicleInputModifier* NewPtr = NewObject<UDefaultModularVehicleInputModifier>(this, InputSetup.InputModifierClass);
-//			InputModifiers.Add(NewPtr);
-//		}
-//		else
-//		{
-//			InputModifiers.Add(nullptr);
-//		}
-//	}
-//}
-
-//void UModularVehicleBaseComponent::ApplyInputModifiers(float DeltaTime, const FModuleInputContainer& RawValue)
-//{
-//	check(InputModifiers.Num() == InputsContainer.GetNumInputs());
-//	for (int I = 0; I < InputsContainer.GetNumInputs(); I++)
-//	{
-//		if (InputModifiers[I] != nullptr)
-//		{
-//			InputsContainer.SetValueAtIndex(I, InputModifiers[I]->InterpInputValue(DeltaTime, InputsContainer.GetValueAtIndex(I), RawValue.GetValueAtIndex(I)));
-//		}
-//		else
-//		{
-//			InputsContainer.SetValueAtIndex(I, RawValue.GetValueAtIndex(I));
-//		}
-//	}
-//}
-
-
 void UModularVehicleBaseComponent::OnCreatePhysicsState()
 {
 	Super::OnCreatePhysicsState();
@@ -705,7 +672,6 @@ void UModularVehicleBaseComponent::SetupInputConfiguration()
 	{
 		ReplicatedState.Container = InputsContainer;
 	}
-	// #TODO reinstate ? GenerateInputModifiers(CombinedInputConfiguration);
 
 	if (VehicleSimulationPT)
 	{
@@ -801,8 +767,6 @@ void UModularVehicleBaseComponent::UpdateState(float DeltaTime)
 	// Should we remove input instead of relying on replicated state in that case?
 	if (bProcessLocally && PVehicleOutput)
 	{
-		//ApplyInputModifiers(DeltaTime, RawInputsContainer); #TODO: If we put this back where does it go
-
 		if (!bUsingNetworkPhysicsPrediction)
 		{
 			// and send to server - (ServerUpdateState_Implementation below)
diff --git a/Engine/Plugins/Experimental/ChaosModularVehicle/Source/ChaosModularVehicleEngine/Public/ChaosModularVehicle/ModularVehicleBaseComponent.h b/Engine/Plugins/Experimental/ChaosModularVehicle/Source/ChaosModularVehicleEngine/Public/ChaosModularVehicle/ModularVehicleBaseComponent.h
--- a/Engine/Plugins/Experimental/ChaosModularVehicle/Source/ChaosModularVehicleEngine/Public/ChaosModularVehicle/ModularVehicleBaseComponent.h
+++ b/Engine/Plugins/Experimental/ChaosModularVehicle/Source/ChaosModularVehicleEngine/Public/ChaosModularVehicle/ModularVehicleBaseComponent.h
@@ -348,8 +348,6 @@ protected:
 	UE_API void AssimilateComponentInputs(TArray<FModuleInputSetup>& OutCombinedInputs);
 	UE_API void CacheRootPhysicsObject(IPhysicsProxyBase* Proxy);
 
-	// #TODO reinstate? virtual void GenerateInputModifiers(const TArray<FModuleInputSetup>& CombinedInputConfiguration);
-	// #TODO reinstate? virtual void ApplyInputModifiers(float DeltaTime, const FModuleInputContainer& RawValue);
 	UE_API void AddSimulationComponentsFromOwnedActor(); // This version does not require cluster unions or geometry collections as parent components
 
 	UE_API virtual void AddComponentToSimulationImpl(UPrimitiveComponent* Component, const TArray<FClusterUnionBoneData>& BonesData, const TArray<FClusterUnionBoneData>& RemovedBoneIDs, bool bIsNew);
diff --git a/Engine/Source/Runtime/Experimental/ChaosVehicles/ChaosVehiclesCore/Private/SimModule/ModuleInput.cpp b/Engine/Source/Runtime/Experimental/ChaosVehicles/ChaosVehiclesCore/Private/SimModule/ModuleInput.cpp
--- a/Engine/Source/Runtime/Experimental/ChaosVehicles/ChaosVehiclesCore/Private/SimModule/ModuleInput.cpp
+++ b/Engine/Source/Runtime/Experimental/ChaosVehicles/ChaosVehiclesCore/Private/SimModule/ModuleInput.cpp
@@ -555,7 +555,7 @@ void FModuleInputContainer::Initialize(TArray<FModuleInputSetup>& SetupData, FIn
 
 	for (FModuleInputSetup& Setup : SetupData)
 	{
-		int Index = AddInput(Setup.Type, Setup.InputModifierClass);
+		int Index = AddInput(Setup.Type);
 		InputValues[Index].SetApplyInputDecay(Setup.bApplyInputDecay);
 		NameMapOut.Add(Setup.Name, Index);
 	}
@@ -591,7 +591,7 @@ void FModuleInputContainer::Serialize(FArchive& Ar, UPackageMap* Map, bool& bOut
 	return;
 }
 
-int FModuleInputContainer::AddInput(EModuleInputValueType Type, TSubclassOf<UDefaultModularVehicleInputModifier>& InputModifierClass)
+int FModuleInputContainer::AddInput(EModuleInputValueType Type)
 {
 	switch(Type)
 	{
diff --git a/Engine/Source/Runtime/Experimental/ChaosVehicles/ChaosVehiclesCore/Private/SimModule/SimModuleTree.cpp b/Engine/Source/Runtime/Experimental/ChaosVehicles/ChaosVehiclesCore/Private/SimModule/SimModuleTree.cpp
--- a/Engine/Source/Runtime/Experimental/ChaosVehicles/ChaosVehiclesCore/Private/SimModule/SimModuleTree.cpp
+++ b/Engine/Source/Runtime/Experimental/ChaosVehicles/ChaosVehiclesCore/Private/SimModule/SimModuleTree.cpp
@@ -115,8 +115,6 @@ void FSimModuleTree::AppendTreeUpdates(const FSimTreeUpdates& TreeUpdates)
 		LocalIndex++;
 	}
 
-	TArray<int> ComponentIndices;
-
 	if (!SimulationModuleTree.IsEmpty())
 	{
 		for (const FPendingModuleDeletions& TreeUpdate : TreeUpdates.GetDeletedModules())
@@ -127,7 +125,6 @@ void FSimModuleTree::AppendTreeUpdates(const FSimTreeUpdates& TreeUpdates)
 				{
 					if (SimModule->GetGuid() == TreeUpdate.Guid)
 					{
-						ComponentIndices.AddUnique(SimModule->GetTransformIndex());
 						DeleteNode(Index);
 						break;
 					}
diff --git a/Engine/Source/Runtime/Experimental/ChaosVehicles/ChaosVehiclesCore/Public/SimModule/ModuleInput.h b/Engine/Source/Runtime/Experimental/ChaosVehicles/ChaosVehiclesCore/Public/SimModule/ModuleInput.h
--- a/Engine/Source/Runtime/Experimental/ChaosVehicles/ChaosVehiclesCore/Public/SimModule/ModuleInput.h
+++ b/Engine/Source/Runtime/Experimental/ChaosVehicles/ChaosVehiclesCore/Public/SimModule/ModuleInput.h
@@ -534,9 +534,6 @@ struct FModuleInputSetup
 	UPROPERTY(EditAnywhere, Category = VehicleInput)
 	EModuleInputValueType Type;
 
-	UPROPERTY(EditAnywhere, Category = VehicleInput)
-	TSubclassOf<UDefaultModularVehicleInputModifier> InputModifierClass;
-
 	/** During physics resimulation, apply decay for this input while it's being extrapolated */
 	UPROPERTY(EditAnywhere, Category = VehicleInput)
 	bool bApplyInputDecay;
@@ -628,7 +625,7 @@ public:
 
 	UE_API void Serialize(FArchive& Ar, UPackageMap* Map, bool& bOutSuccess, EModuleInputQuantizationType InInputQuantizationType);
 
-	UE_API int AddInput(EModuleInputValueType Type, TSubclassOf<UDefaultModularVehicleInputModifier>& InputModifierClass);
+	UE_API int AddInput(EModuleInputValueType Type);
 
 	UE_API void RemoveAllInputs();
 
