Modular Vehicle : Add custom enhanced input modifier to match OG chaos vehicle plugin

diff --git a/Engine/Plugins/Experimental/ChaosModularVehicle/ChaosModularVehicle.uplugin b/Engine/Plugins/Experimental/ChaosModularVehicle/ChaosModularVehicle.uplugin
--- a/Engine/Plugins/Experimental/ChaosModularVehicle/ChaosModularVehicle.uplugin
+++ b/Engine/Plugins/Experimental/ChaosModularVehicle/ChaosModularVehicle.uplugin
@@ -40,6 +40,10 @@
 		{
 			"Name": "GeometryCollectionPlugin",
 			"Enabled": true
+		},
+		{
+			"Name": "EnhancedInput",
+			"Enabled": true
 		}
 	]
 
diff --git a/Engine/Plugins/Experimental/ChaosModularVehicle/Source/ChaosModularVehicleEngine/ChaosModularVehicleEngine.Build.cs b/Engine/Plugins/Experimental/ChaosModularVehicle/Source/ChaosModularVehicleEngine/ChaosModularVehicleEngine.Build.cs
--- a/Engine/Plugins/Experimental/ChaosModularVehicle/Source/ChaosModularVehicleEngine/ChaosModularVehicleEngine.Build.cs
+++ b/Engine/Plugins/Experimental/ChaosModularVehicle/Source/ChaosModularVehicleEngine/ChaosModularVehicleEngine.Build.cs
@@ -17,6 +17,7 @@ namespace UnrealBuildTool.Rules
 					"Engine",
 					"RenderCore",
 					"RHI",
+					"EnhancedInput",
 					"AnimGraphRuntime",
 					"Renderer",
 					"ChaosVehiclesCore",
diff --git a/Engine/Plugins/Experimental/ChaosModularVehicle/Source/ChaosModularVehicleEngine/Private/ChaosModularVehicle/InputModifier_ModularVehicleSmooth.cpp b/Engine/Plugins/Experimental/ChaosModularVehicle/Source/ChaosModularVehicleEngine/Private/ChaosModularVehicle/InputModifier_ModularVehicleSmooth.cpp
new file mode 100644
--- /dev/null
+++ b/Engine/Plugins/Experimental/ChaosModularVehicle/Source/ChaosModularVehicleEngine/Private/ChaosModularVehicle/InputModifier_ModularVehicleSmooth.cpp
@@ -0,0 +1,29 @@
+// Copyright Epic Games, Inc. All Rights Reserved.
+
+#include "ChaosModularVehicle/InputModifier_ModularVehicleSmooth.h"
+
+FInputActionValue UInputModifier_ModularVehicleSmooth::ModifyRaw_Implementation(
+	const UEnhancedPlayerInput* PlayerInput, 
+	FInputActionValue CurrentValue, 
+	float DeltaTime)
+{
+	// Only handle scalar (float) inputs here.
+	float Current = CurrentValue.Get<float>();
+
+	// No Fall rate just snaps to zero
+	if (FMath::Abs(Current) < SMALL_NUMBER)
+	{
+		PreviousValue = 0.0f;
+		return 0.0f;
+	}
+
+	float DeltaValue = Current - PreviousValue;
+
+	const float MaxDeltaValue = DeltaTime * RiseRate;
+	const float ClampedDeltaValue = FMath::Clamp(DeltaValue, -MaxDeltaValue, MaxDeltaValue);
+	
+	float SmoothedValue = PreviousValue + ClampedDeltaValue;
+	PreviousValue = SmoothedValue;
+
+	return SmoothedValue;
+}
diff --git a/Engine/Plugins/Experimental/ChaosModularVehicle/Source/ChaosModularVehicleEngine/Public/ChaosModularVehicle/InputModifier_ModularVehicleSmooth.h b/Engine/Plugins/Experimental/ChaosModularVehicle/Source/ChaosModularVehicleEngine/Public/ChaosModularVehicle/InputModifier_ModularVehicleSmooth.h
new file mode 100644
--- /dev/null
+++ b/Engine/Plugins/Experimental/ChaosModularVehicle/Source/ChaosModularVehicleEngine/Public/ChaosModularVehicle/InputModifier_ModularVehicleSmooth.h
@@ -0,0 +1,33 @@
+// Copyright Epic Games, Inc. All Rights Reserved.
+
+ #pragma once
+
+#include "CoreMinimal.h"
+#include "InputModifiers.h"
+#include "InputActionValue.h"
+#include "InputModifier_ModularVehicleSmooth.generated.h"
+
+/**
+ * UInputModifier_ModularVehicleSmooth
+ *
+ * Add this as a Modifier on any float‐based Action (e.g. “Steer”) to get a
+ * smoother, “lagged” feel. To match the behaviour of the OG chaos vehicle plugin
+ */
+UCLASS(BlueprintType, meta = (DisplayName = "Modular Vehicle Smooth Modifier"))
+class CHAOSMODULARVEHICLEENGINE_API UInputModifier_ModularVehicleSmooth : public UInputModifier
+{
+	GENERATED_BODY()
+
+public:
+	/**
+	 * How quickly to “catch up” to the raw input.
+	 * Higher numbers = snappier response; lower = more lag/smoothing.
+	 */
+	UPROPERTY(EditAnywhere, BlueprintReadWrite, Category = "Smoothing")
+	float RiseRate = 10.0f;
+
+protected:
+	virtual FInputActionValue ModifyRaw_Implementation(const UEnhancedPlayerInput* PlayerInput, FInputActionValue CurrentValue, float DeltaTime) override;
+
+	float PreviousValue = 0.0f;
+};
\ No newline at end of file
