Chaos Modular Vehicle:
- Add additional param to  simulate callback for async input.
- Expose functionality in UModularVehicleBaseComponent for adding inputs to the input container from code.
- Do not simulate the tree if the number of modules is 0.
- In case of sim processing order being set to ManualOverride, the user needs to perform the suspension raycasts/additional work as well.

diff --git a/Engine/Plugins/Experimental/ChaosModularVehicle/Source/ChaosModularVehicleEngine/Private/ChaosModularVehicle/ModularVehicleBaseComponent.cpp b/Engine/Plugins/Experimental/ChaosModularVehicle/Source/ChaosModularVehicleEngine/Private/ChaosModularVehicle/ModularVehicleBaseComponent.cpp
--- a/Engine/Plugins/Experimental/ChaosModularVehicle/Source/ChaosModularVehicleEngine/Private/ChaosModularVehicle/ModularVehicleBaseComponent.cpp
+++ b/Engine/Plugins/Experimental/ChaosModularVehicle/Source/ChaosModularVehicleEngine/Private/ChaosModularVehicle/ModularVehicleBaseComponent.cpp
@@ -1706,6 +1706,12 @@ void UModularVehicleBaseComponent::BroadcastModuleRemovedEvent(const FName& SimT
 	}
 }
 
+void UModularVehicleBaseComponent::AddInput(FModuleInputSetup InputSetup)
+{
+	InputConfig.Add(InputSetup);
+}
+
+
 
 //-=====================================================
 // Networking Replication
diff --git a/Engine/Plugins/Experimental/ChaosModularVehicle/Source/ChaosModularVehicleEngine/Private/ChaosModularVehicle/ModularVehicleSimulationCU.cpp b/Engine/Plugins/Experimental/ChaosModularVehicle/Source/ChaosModularVehicleEngine/Private/ChaosModularVehicle/ModularVehicleSimulationCU.cpp
--- a/Engine/Plugins/Experimental/ChaosModularVehicle/Source/ChaosModularVehicleEngine/Private/ChaosModularVehicle/ModularVehicleSimulationCU.cpp
+++ b/Engine/Plugins/Experimental/ChaosModularVehicle/Source/ChaosModularVehicleEngine/Private/ChaosModularVehicle/ModularVehicleSimulationCU.cpp
@@ -77,7 +77,10 @@ void FModularVehicleSimulation::SimulateModuleTree(UWorld* InWorld, float DeltaS
 	if (Proxy && SimModuleTree.IsValid())
 	{
 		int InitialNum = SimModuleTree->GetSimulationModuleTree().Num();
-
+		if (InitialNum == 0)
+		{
+			return;
+		}
 		//if (InWorld)
 		//{
 		//	WriteNetReport(InWorld->IsNetMode(NM_Client), FString::Printf(TEXT("X %s,  R %s,  V %s,  W %s")
@@ -124,13 +127,17 @@ void FModularVehicleSimulation::SimulateModuleTree(UWorld* InWorld, float DeltaS
 		SimInputData.StateInputs = &StateInterface;
 		SimInputData.bKeepVehicleAwake = InputData.PhysicsInputs.NetworkInputs.VehicleInputs.KeepAwake;
 
-		PerformAdditionalSimWork(InWorld, InputData, Proxy, SimInputData);
 
+		if(SimModuleTree->GetSimTreeProcessingOrder() != ESimTreeProcessingOrder::ManualOverride)
+		{
+			PerformAdditionalSimWork(InWorld, InputData, Proxy, SimInputData);
+		}
 		// run the dynamics simulation, engine, suspension, wheels, aerofoils etc.
 		SimModuleTree->Simulate(DeltaSeconds, SimInputData, Proxy, RootParticle);
+
 		if(SimModuleTree->GetSimTreeProcessingOrder() == ESimTreeProcessingOrder::ManualOverride)
 		{
-			VehicleSimulationCallback.Broadcast(DeltaSeconds, SimInputData, Proxy, RootParticle, SimModuleTree.Get());
+			VehicleSimulationCallback.Broadcast(DeltaSeconds, InputData, SimInputData, Proxy, RootParticle, SimModuleTree.Get());
 		}
 	}
 
diff --git a/Engine/Plugins/Experimental/ChaosModularVehicle/Source/ChaosModularVehicleEngine/Public/ChaosModularVehicle/ModularVehicleBaseComponent.h b/Engine/Plugins/Experimental/ChaosModularVehicle/Source/ChaosModularVehicleEngine/Public/ChaosModularVehicle/ModularVehicleBaseComponent.h
--- a/Engine/Plugins/Experimental/ChaosModularVehicle/Source/ChaosModularVehicleEngine/Public/ChaosModularVehicle/ModularVehicleBaseComponent.h
+++ b/Engine/Plugins/Experimental/ChaosModularVehicle/Source/ChaosModularVehicleEngine/Public/ChaosModularVehicle/ModularVehicleBaseComponent.h
@@ -408,6 +408,8 @@ public:
 	UFUNCTION(reliable, server, WithValidation)
 	UE_API void ServerUpdateState(const FModuleInputContainer& InputsIn, bool KeepAwake);
 
+	UE_API void AddInput(FModuleInputSetup InputSetup);
+
 	UE_API void LogInputSetup();
 
 	UE_API int AddSimModule(Chaos::ISimulationModuleBase* NewModule, const FTransform& ComponentTransform, int ParentIndex, int TransformIndex);
diff --git a/Engine/Plugins/Experimental/ChaosModularVehicle/Source/ChaosModularVehicleEngine/Public/ChaosModularVehicle/ModularVehicleSimulationCU.h b/Engine/Plugins/Experimental/ChaosModularVehicle/Source/ChaosModularVehicleEngine/Public/ChaosModularVehicle/ModularVehicleSimulationCU.h
--- a/Engine/Plugins/Experimental/ChaosModularVehicle/Source/ChaosModularVehicleEngine/Public/ChaosModularVehicle/ModularVehicleSimulationCU.h
+++ b/Engine/Plugins/Experimental/ChaosModularVehicle/Source/ChaosModularVehicleEngine/Public/ChaosModularVehicle/ModularVehicleSimulationCU.h
@@ -37,7 +37,7 @@ namespace Chaos
 {
 	class FClusterUnionPhysicsProxy;
 	class FSingleParticlePhysicsProxy;
-	DECLARE_MULTICAST_DELEGATE_FiveParams(FModularVehicleSimulationCallback, float, FAllInputs&, IPhysicsProxyBase*, Chaos::FPBDRigidParticleHandle*, Chaos::FSimModuleTree*);
+	DECLARE_MULTICAST_DELEGATE_SixParams(FModularVehicleSimulationCallback, float, const FModularVehicleAsyncInput&, FAllInputs&, IPhysicsProxyBase*, Chaos::FPBDRigidParticleHandle*, Chaos::FSimModuleTree*);
 	
 }
 
