Adding Title Property to Input Config and State Input Configuration UPROPERTY's making it easier to find in the Editor without expanding everything.

diff --git a/Engine/Plugins/Experimental/ChaosModularVehicle/Source/ChaosModularVehicleEngine/Public/ChaosModularVehicle/ModularVehicleBaseComponent.h b/Engine/Plugins/Experimental/ChaosModularVehicle/Source/ChaosModularVehicleEngine/Public/ChaosModularVehicle/ModularVehicleBaseComponent.h
--- a/Engine/Plugins/Experimental/ChaosModularVehicle/Source/ChaosModularVehicleEngine/Public/ChaosModularVehicle/ModularVehicleBaseComponent.h
+++ b/Engine/Plugins/Experimental/ChaosModularVehicle/Source/ChaosModularVehicleEngine/Public/ChaosModularVehicle/ModularVehicleBaseComponent.h
@@ -392,10 +392,10 @@ protected:
 
 public:
 
-	UPROPERTY(EditAnywhere, Category = VehicleInput)
+	UPROPERTY(EditAnywhere, Category = VehicleInput, meta=(TitleProperty="{Name} | {Type}"))
 	TArray<FModuleInputSetup> InputConfig;
 
-	UPROPERTY(EditAnywhere, Category = VehicleState)
+	UPROPERTY(EditAnywhere, Category = VehicleState, meta=(TitleProperty="{Name} | {Type}"))
 	TArray<FModuleInputSetup> StateInputConfiguration;
 		
 	UPROPERTY(EditAnywhere, Category = "Game|Components|ModularVehicle")
