Modular Vehicle: Shutdown crash fix

diff --git a/Engine/Plugins/Experimental/ChaosModularVehicle/Source/ChaosModularVehicleEngine/Private/ChaosModularVehicle/ModularVehicleBaseComponent.cpp b/Engine/Plugins/Experimental/ChaosModularVehicle/Source/ChaosModularVehicleEngine/Private/ChaosModularVehicle/ModularVehicleBaseComponent.cpp
--- a/Engine/Plugins/Experimental/ChaosModularVehicle/Source/ChaosModularVehicleEngine/Private/ChaosModularVehicle/ModularVehicleBaseComponent.cpp
+++ b/Engine/Plugins/Experimental/ChaosModularVehicle/Source/ChaosModularVehicleEngine/Private/ChaosModularVehicle/ModularVehicleBaseComponent.cpp
@@ -1103,47 +1103,50 @@ void UModularVehicleBaseComponent::ActionTreeUpdates(Chaos::FSimTreeUpdates* Nex
 
 			Solver->EnqueueCommandImmediate([Proxy, this, NextTreeUpdates = *NextTreeUpdates]() mutable
 			{
-				if (IsValidChecked(this) && bPhysicsStateCreated && VehicleSimulationPT)
+				if (Chaos::FPhysicsSolver* Solver = Proxy->GetSolver<Chaos::FPhysicsSolver>())
 				{
-					TUniquePtr<Chaos::FSimModuleTree>& SimModuleTree = VehicleSimulationPT->AccessSimComponentTree();
-					if(SimModuleTree.IsValid())
+					if (!Solver->IsShuttingDown() && IsValidChecked(this) && bPhysicsStateCreated && VehicleSimulationPT)
 					{
-						SimModuleTree->AppendTreeUpdates(NextTreeUpdates);
-						FModularVehicleBuilder::FixupTreeLinks(SimModuleTree);
-
-#if !(UE_BUILD_SHIPPING || UE_BUILD_TEST)
-						if (bModularVehicle_DumpModuleTreeStructure_Enabled)
+						TUniquePtr<Chaos::FSimModuleTree>& SimModuleTree = VehicleSimulationPT->AccessSimComponentTree();
+						if(SimModuleTree.IsValid())
 						{
-							UE_LOG(LogModularBase, Warning, TEXT("SimTreeModules:") );
-							for (int I = 0; I < SimModuleTree->GetNumNodes(); I++)
+							SimModuleTree->AppendTreeUpdates(NextTreeUpdates);
+							FModularVehicleBuilder::FixupTreeLinks(SimModuleTree);
+
+	#if !(UE_BUILD_SHIPPING || UE_BUILD_TEST)
+							if (bModularVehicle_DumpModuleTreeStructure_Enabled)
 							{
-								if (Chaos::ISimulationModuleBase* Module = SimModuleTree->GetNode(I).SimModule)
+								UE_LOG(LogModularBase, Warning, TEXT("SimTreeModules:") );
+								for (int I = 0; I < SimModuleTree->GetNumNodes(); I++)
 								{
-									FString String;
-									Module->GetDebugString(String);
-									UE_LOG(LogModularBase, Warning, TEXT("..%s"), *String);
+									if (Chaos::ISimulationModuleBase* Module = SimModuleTree->GetNode(I).SimModule)
+									{
+										FString String;
+										Module->GetDebugString(String);
+										UE_LOG(LogModularBase, Warning, TEXT("..%s"), *String);
+									}
 								}
 							}
-						}
-#endif // !(UE_BUILD_SHIPPING || UE_BUILD_TEST)
+	#endif // !(UE_BUILD_SHIPPING || UE_BUILD_TEST)
 
-						// Network replication data needs to be updated, this is currently studily slow
-						if (NetworkPhysicsComponent)
-						{
-							TSharedPtr<Chaos::FBaseRewindHistory>& History = NetworkPhysicsComponent->GetStateHistory_Internal();
-							Chaos::TDataRewindHistory<FNetworkModularVehicleStates>* StateHistory = static_cast<Chaos::TDataRewindHistory<FNetworkModularVehicleStates>*>(History.Get());
-							if (StateHistory)
+							// Network replication data needs to be updated, this is currently studily slow
+							if (NetworkPhysicsComponent)
 							{
-								// #TODO: we are rebuilding from scratch every time there is a single change, there must be a better way!
-								// not sure of it is safe to update the data at this time?
-								for (int I = 0; I < StateHistory->GetDataHistory().Num(); I++)
+								TSharedPtr<Chaos::FBaseRewindHistory>& History = NetworkPhysicsComponent->GetStateHistory_Internal();
+								Chaos::TDataRewindHistory<FNetworkModularVehicleStates>* StateHistory = static_cast<Chaos::TDataRewindHistory<FNetworkModularVehicleStates>*>(History.Get());
+								if (StateHistory)
 								{
-									FNetworkModularVehicleStates& State = StateHistory->GetDataHistory()[I];
-									State.ModuleData.Empty();
-									TUniquePtr<Chaos::FSimModuleTree>& InnerSimModuleTree = VehicleSimulationPT->AccessSimComponentTree();
-									if(InnerSimModuleTree.IsValid())
+									// #TODO: we are rebuilding from scratch every time there is a single change, there must be a better way!
+									// not sure of it is safe to update the data at this time?
+									for (int I = 0; I < StateHistory->GetDataHistory().Num(); I++)
 									{
-										InnerSimModuleTree->GenerateReplicationStructure(State.ModuleData);
+										FNetworkModularVehicleStates& State = StateHistory->GetDataHistory()[I];
+										State.ModuleData.Empty();
+										TUniquePtr<Chaos::FSimModuleTree>& InnerSimModuleTree = VehicleSimulationPT->AccessSimComponentTree();
+										if(InnerSimModuleTree.IsValid())
+										{
+											InnerSimModuleTree->GenerateReplicationStructure(State.ModuleData);
+										}
 									}
 								}
 							}
