ModularVehicle: Threading issue fix

diff --git a/Engine/Plugins/Experimental/ChaosModularVehicle/Source/ChaosModularVehicleEngine/Private/ChaosModularVehicle/ChaosSimModuleManagerAsyncCallback.cpp b/Engine/Plugins/Experimental/ChaosModularVehicle/Source/ChaosModularVehicleEngine/Private/ChaosModularVehicle/ChaosSimModuleManagerAsyncCallback.cpp
--- a/Engine/Plugins/Experimental/ChaosModularVehicle/Source/ChaosModularVehicleEngine/Private/ChaosModularVehicle/ChaosSimModuleManagerAsyncCallback.cpp
+++ b/Engine/Plugins/Experimental/ChaosModularVehicle/Source/ChaosModularVehicleEngine/Private/ChaosModularVehicle/ChaosSimModuleManagerAsyncCallback.cpp
@@ -192,13 +192,16 @@ TUniquePtr<FModularVehicleAsyncOutput> FModularVehicleAsyncInput::Simulate(UWorl
 		return Output;
 	}
 
-	if (Vehicle && Vehicle->VehicleSimulationPT)
+	if (Vehicle)
 	{
-		// FILL OUTPUT DATA HERE THAT WILL GET PASSED BACK TO THE GAME THREAD
-		Vehicle->VehicleSimulationPT->Simulate(World, DeltaSeconds, *this, *Output.Get(), Proxy);
+		if (FModularVehicleSimulation* Sim = Vehicle->VehicleSimulationPT.Get())
+		{
+			// FILL OUTPUT DATA HERE THAT WILL GET PASSED BACK TO THE GAME THREAD
+			Sim->Simulate(World, DeltaSeconds, *this, *Output.Get(), Proxy);
 
-		FModularVehicleAsyncOutput& OutputData = *Output.Get();
-		Vehicle->VehicleSimulationPT->FillOutputState(OutputData);
+			FModularVehicleAsyncOutput& OutputData = *Output.Get();
+			Sim->FillOutputState(OutputData);
+		}
 	}
 
 
diff --git a/Engine/Plugins/Experimental/ChaosModularVehicle/Source/ChaosModularVehicleEngine/Private/ChaosModularVehicle/ModularVehicleBaseComponent.cpp b/Engine/Plugins/Experimental/ChaosModularVehicle/Source/ChaosModularVehicleEngine/Private/ChaosModularVehicle/ModularVehicleBaseComponent.cpp
--- a/Engine/Plugins/Experimental/ChaosModularVehicle/Source/ChaosModularVehicleEngine/Private/ChaosModularVehicle/ModularVehicleBaseComponent.cpp
+++ b/Engine/Plugins/Experimental/ChaosModularVehicle/Source/ChaosModularVehicleEngine/Private/ChaosModularVehicle/ModularVehicleBaseComponent.cpp
@@ -1373,6 +1373,8 @@ void UModularVehicleBaseComponent::CreateVehicleSim()
 	{
 		if (World->IsGameWorld())
 		{
+			UE_LOG(LogModularBase, Log, TEXT("UModularVehicleBaseComponent::CreateVehicleSim"));
+
 			// create the simulation class
 			VehicleSimulationPT = MakeUnique<FModularVehicleSimulation>(bUsingNetworkPhysicsPrediction, (int8)World->GetNetMode());
 
@@ -1398,6 +1400,8 @@ void UModularVehicleBaseComponent::DestroyVehicleSim()
 	UWorld* World = GetWorld();
 	if (World->IsGameWorld())
 	{
+		UE_LOG(LogModularBase, Log, TEXT("UModularVehicleBaseComponent::DestroyVehicleSim"));
+
 		FPhysScene* PhysScene = World->GetPhysicsScene();
 		check(PhysScene);
 
@@ -1413,12 +1417,17 @@ void UModularVehicleBaseComponent::DestroyVehicleSim()
 				PVehicleOutput.Reset(nullptr);
 			}
 
-			// This is a managed pointer, so do we need to do this?
-			if (VehicleSimulationPT)
-			{
-				VehicleSimulationPT->Terminate();
-				VehicleSimulationPT.Reset(nullptr);
-			}
+			// postpone deletion of object while simultaneously setting it to null in the component (TUniquePtr)
+			LocalSolver->EnqueueCommandImmediate([Pointer = MoveTemp(VehicleSimulationPT)]() mutable
+				{
+					Pointer.Reset(nullptr);
+				});
+
+		}
+
+		if (bUsingNetworkPhysicsPrediction && NetworkPhysicsComponent)
+		{
+			NetworkPhysicsComponent->RemoveDataHistory();
 		}
 
 	}
diff --git a/Engine/Plugins/Experimental/ChaosModularVehicle/Source/ChaosModularVehicleEngine/Private/ChaosModularVehicle/ModularVehicleBuilder.cpp b/Engine/Plugins/Experimental/ChaosModularVehicle/Source/ChaosModularVehicleEngine/Private/ChaosModularVehicle/ModularVehicleBuilder.cpp
--- a/Engine/Plugins/Experimental/ChaosModularVehicle/Source/ChaosModularVehicleEngine/Private/ChaosModularVehicle/ModularVehicleBuilder.cpp
+++ b/Engine/Plugins/Experimental/ChaosModularVehicle/Source/ChaosModularVehicleEngine/Private/ChaosModularVehicle/ModularVehicleBuilder.cpp
@@ -35,61 +35,64 @@ void FModularVehicleBuilder::FixupTreeLinks(TUniquePtr<Chaos::FSimModuleTree>& S
 {
 	using namespace Chaos;
 
-	const FSimModuleTree::FSimModuleNode* TransmissionNode = SimModuleTree->LocateNodeByType<Chaos::FTransmissionSimModule>();
-
-	for (int I = 0; I < SimModuleTree->NumActiveNodes(); I++)
+	if (SimModuleTree.IsValid())
 	{
-		// link suspension and wheel references to one another
-		ISimulationModuleBase* Module = SimModuleTree->AccessSimModule(I);
-		if (Module && Module->IsSimType<FSuspensionBaseInterface>())
-		{
-			FSimModuleTree::FSimModuleNode& SuspensionNode = SimModuleTree->GetNode(I);
+		const FSimModuleTree::FSimModuleNode* TransmissionNode = SimModuleTree->LocateNodeByType<Chaos::FTransmissionSimModule>();
 
-			if (SuspensionNode.Parent != FSimModuleTree::FSimModuleNode::INVALID_IDX)
+		for (int I = 0; I < SimModuleTree->NumActiveNodes(); I++)
+		{
+			// link suspension and wheel references to one another
+			ISimulationModuleBase* Module = SimModuleTree->AccessSimModule(I);
+			if (Module && Module->IsSimType<FSuspensionBaseInterface>())
 			{
-				FSimModuleTree::FSimModuleNode& SuspensionParentNode = SimModuleTree->GetNode(SuspensionNode.Parent);
+				FSimModuleTree::FSimModuleNode& SuspensionNode = SimModuleTree->GetNode(I);
 
-				ISimulationModuleBase* ParentModule = SimModuleTree->AccessSimModule(SuspensionParentNode.SimModule->GetTreeIndex());
-
-				if(FWheelBaseInterface* Wheel = ParentModule->Cast<FWheelBaseInterface>())
+				if (SuspensionNode.Parent != FSimModuleTree::FSimModuleNode::INVALID_IDX)
 				{
-					if (FSuspensionBaseInterface* Suspension = Module->Cast<FSuspensionBaseInterface>())
+					FSimModuleTree::FSimModuleNode& SuspensionParentNode = SimModuleTree->GetNode(SuspensionNode.Parent);
+
+					ISimulationModuleBase* ParentModule = SimModuleTree->AccessSimModule(SuspensionParentNode.SimModule->GetTreeIndex());
+
+					if (FWheelBaseInterface* Wheel = ParentModule->Cast<FWheelBaseInterface>())
 					{
-						Wheel->SetSuspensionSimTreeIndex(Suspension->GetTreeIndex());
-						Suspension->SetWheelSimTreeIndex(Wheel->GetTreeIndex());
+						if (FSuspensionBaseInterface* Suspension = Module->Cast<FSuspensionBaseInterface>())
+						{
+							Wheel->SetSuspensionSimTreeIndex(Suspension->GetTreeIndex());
+							Suspension->SetWheelSimTreeIndex(Wheel->GetTreeIndex());
+						}
 					}
 				}
-			}
 
-			for (int Child : SuspensionNode.Children)
-			{
-				if (Child != FSimModuleTree::FSimModuleNode::INVALID_IDX)
+				for (int Child : SuspensionNode.Children)
 				{
-					FSimModuleTree::FSimModuleNode& SuspensionChildNode = SimModuleTree->GetNode(Child);
+					if (Child != FSimModuleTree::FSimModuleNode::INVALID_IDX)
+					{
+						FSimModuleTree::FSimModuleNode& SuspensionChildNode = SimModuleTree->GetNode(Child);
 
-					ISimulationModuleBase* ChildModule = SimModuleTree->AccessSimModule(SuspensionChildNode.SimModule->GetTreeIndex());
+						ISimulationModuleBase* ChildModule = SimModuleTree->AccessSimModule(SuspensionChildNode.SimModule->GetTreeIndex());
 
-					if(FWheelBaseInterface* Wheel = ChildModule->Cast<FWheelBaseInterface>())
-					{
-						if (FSuspensionBaseInterface* Suspension = Module->Cast<FSuspensionBaseInterface>())
+						if (FWheelBaseInterface* Wheel = ChildModule->Cast<FWheelBaseInterface>())
 						{
-							Wheel->SetSuspensionSimTreeIndex(Suspension->GetTreeIndex());
-							Suspension->SetWheelSimTreeIndex(Wheel->GetTreeIndex());
+							if (FSuspensionBaseInterface* Suspension = Module->Cast<FSuspensionBaseInterface>())
+							{
+								Wheel->SetSuspensionSimTreeIndex(Suspension->GetTreeIndex());
+								Suspension->SetWheelSimTreeIndex(Wheel->GetTreeIndex());
+							}
 						}
 					}
 				}
 			}
-		}
 
-		// temporarily - link suspension/wheels to engine/clutch/transmission, assuming only 1 transmission & all wheels are powered
-		if (TransmissionNode)
-		{
-			if (Module && Module->IsSimType<class FWheelSimModule>())
+			// temporarily - link suspension/wheels to engine/clutch/transmission, assuming only 1 transmission & all wheels are powered
+			if (TransmissionNode)
 			{
-				SimModuleTree->Reparent(Module->GetTreeIndex(), TransmissionNode->SimModule->GetTreeIndex());
+				if (Module && Module->IsSimType<class FWheelSimModule>())
+				{
+					SimModuleTree->Reparent(Module->GetTreeIndex(), TransmissionNode->SimModule->GetTreeIndex());
+				}
 			}
-		}
 
+		}
 	}
 }
 
diff --git a/Engine/Plugins/Experimental/ChaosModularVehicle/Source/ChaosModularVehicleEngine/Private/ChaosModularVehicle/ModularVehicleSimulationCU.cpp b/Engine/Plugins/Experimental/ChaosModularVehicle/Source/ChaosModularVehicleEngine/Private/ChaosModularVehicle/ModularVehicleSimulationCU.cpp
--- a/Engine/Plugins/Experimental/ChaosModularVehicle/Source/ChaosModularVehicleEngine/Private/ChaosModularVehicle/ModularVehicleSimulationCU.cpp
+++ b/Engine/Plugins/Experimental/ChaosModularVehicle/Source/ChaosModularVehicleEngine/Private/ChaosModularVehicle/ModularVehicleSimulationCU.cpp
@@ -33,11 +33,16 @@ FAutoConsoleVariableRef CVarChaosModularVehiclesDisableAnim(TEXT("p.ModularVehic
 
 void FModularVehicleSimulation::Initialize(TUniquePtr<Chaos::FSimModuleTree>& InSimModuleTree)
 {
+	UE_LOG(LogModularVehicleSim, Log, TEXT("FModularVehicleSimulation::Initialize"));
+
 	SimModuleTree = MoveTemp(InSimModuleTree);
 }
 
 void FModularVehicleSimulation::Terminate()
 {
+	UE_LOG(LogModularVehicleSim, Log, TEXT("FModularVehicleSimulation::Terminate"));
+
+	RootParticle = nullptr;
 	SimModuleTree.Reset(nullptr);
 }
 
diff --git a/Engine/Plugins/Experimental/ChaosModularVehicle/Source/ChaosModularVehicleEngine/Public/ChaosModularVehicle/ModularVehicleSimulationCU.h b/Engine/Plugins/Experimental/ChaosModularVehicle/Source/ChaosModularVehicleEngine/Public/ChaosModularVehicle/ModularVehicleSimulationCU.h
--- a/Engine/Plugins/Experimental/ChaosModularVehicle/Source/ChaosModularVehicleEngine/Public/ChaosModularVehicle/ModularVehicleSimulationCU.h
+++ b/Engine/Plugins/Experimental/ChaosModularVehicle/Source/ChaosModularVehicleEngine/Public/ChaosModularVehicle/ModularVehicleSimulationCU.h
@@ -55,7 +55,7 @@ public:
 
 	virtual ~FModularVehicleSimulation()
 	{
-		SimModuleTree.Reset();
+		Terminate();
 	}
 
 	UE_API void Initialize(TUniquePtr<Chaos::FSimModuleTree>& InSimModuleTree);
