Fixes a few cases where physics solver time was being cached as a 32-bit float and possibly losing precision for long-running simulations. Also removed a couple of copy-paste cases where the solver time was queried but never actually used.

#rnx

diff --git a/Engine/Plugins/Experimental/ChaosModularVehicle/Source/ChaosModularVehicleEngine/Private/ChaosModularVehicle/ChaosSimModuleManager.cpp b/Engine/Plugins/Experimental/ChaosModularVehicle/Source/ChaosModularVehicleEngine/Private/ChaosModularVehicle/ChaosSimModuleManager.cpp
--- a/Engine/Plugins/Experimental/ChaosModularVehicle/Source/ChaosModularVehicleEngine/Private/ChaosModularVehicle/ChaosSimModuleManager.cpp
+++ b/Engine/Plugins/Experimental/ChaosModularVehicle/Source/ChaosModularVehicleEngine/Private/ChaosModularVehicle/ChaosSimModuleManager.cpp
@@ -278,7 +278,7 @@ void FChaosSimModuleManager::ParallelUpdateVehicles(float DeltaSeconds)
 	}
 
 	// Since we are in pre-physics, delta seconds is not accounted for in external time yet
-	const float ResultsTime = AsyncCallback->GetSolver()->GetPhysicsResultsTime_External();
+	const double ResultsTime = AsyncCallback->GetSolver()->GetPhysicsResultsTime_External();
 
 	// Find index of first non-consumable output (first one after current time)
 	int32 LastOutputIdx = 0;
diff --git a/Engine/Plugins/Experimental/ChaosVehiclesPlugin/Source/ChaosVehicles/Private/ChaosVehicleManager.cpp b/Engine/Plugins/Experimental/ChaosVehiclesPlugin/Source/ChaosVehicles/Private/ChaosVehicleManager.cpp
--- a/Engine/Plugins/Experimental/ChaosVehiclesPlugin/Source/ChaosVehicles/Private/ChaosVehicleManager.cpp
+++ b/Engine/Plugins/Experimental/ChaosVehiclesPlugin/Source/ChaosVehicles/Private/ChaosVehicleManager.cpp
@@ -223,7 +223,7 @@ void FChaosVehicleManager::ParallelUpdateVehicles(float DeltaSeconds)
 	}
 
 	// Since we are in pre-physics, delta seconds is not accounted for in external time yet
-	const float ResultsTime = AsyncCallback->GetSolver()->GetPhysicsResultsTime_External() + DeltaSeconds;
+	const double ResultsTime = AsyncCallback->GetSolver()->GetPhysicsResultsTime_External() + DeltaSeconds;
 
 	// Find index of first non-consumable output (first one after current time)
 	int32 LastOutputIdx = 0;
diff --git a/Engine/Plugins/Experimental/Water/Source/Runtime/Private/BuoyancyManager.cpp b/Engine/Plugins/Experimental/Water/Source/Runtime/Private/BuoyancyManager.cpp
--- a/Engine/Plugins/Experimental/Water/Source/Runtime/Private/BuoyancyManager.cpp
+++ b/Engine/Plugins/Experimental/Water/Source/Runtime/Private/BuoyancyManager.cpp
@@ -161,7 +161,7 @@ void ABuoyancyManager::Update(FPhysScene* PhysScene, float DeltaTime)
 			PendingOutputs.Emplace(MoveTemp(AsyncOutputLatest));
 		}
 
-		const float ResultsTime = AsyncCallback->GetSolver()->GetPhysicsResultsTime_External() + DeltaTime;
+		const double ResultsTime = AsyncCallback->GetSolver()->GetPhysicsResultsTime_External() + DeltaTime;
 
 		// Find index of first non-consumable output (first one after current time)
 		int32 LastOutputIdx = 0;
