Removed unnecessary dependencies from Chaos Modular Vehicle plugin

#rnx
#tests PIE, preflights

diff --git a/Engine/Plugins/Experimental/ChaosModularVehicle/ChaosModularVehicle.uplugin b/Engine/Plugins/Experimental/ChaosModularVehicle/ChaosModularVehicle.uplugin
--- a/Engine/Plugins/Experimental/ChaosModularVehicle/ChaosModularVehicle.uplugin
+++ b/Engine/Plugins/Experimental/ChaosModularVehicle/ChaosModularVehicle.uplugin
@@ -33,14 +33,6 @@
 
 	],
 	"Plugins": [
-		{
-			"Name": "ProceduralMeshComponent",
-			"Enabled": true
-		},
-		{
-			"Name": "GeometryCollectionPlugin",
-			"Enabled": true
-		},
 		{
 			"Name": "EnhancedInput",
 			"Enabled": true
