[ModularVehicle] Fixes up invalid nettoken hashing and initialization.

diff --git a/Engine/Plugins/Experimental/ChaosModularVehicle/Source/ChaosModularVehicleEngine/Private/ChaosModularVehicle/ChaosSimModuleManagerAsyncCallback.cpp b/Engine/Plugins/Experimental/ChaosModularVehicle/Source/ChaosModularVehicleEngine/Private/ChaosModularVehicle/ChaosSimModuleManagerAsyncCallback.cpp
--- a/Engine/Plugins/Experimental/ChaosModularVehicle/Source/ChaosModularVehicleEngine/Private/ChaosModularVehicle/ChaosSimModuleManagerAsyncCallback.cpp
+++ b/Engine/Plugins/Experimental/ChaosModularVehicle/Source/ChaosModularVehicleEngine/Private/ChaosModularVehicle/ChaosSimModuleManagerAsyncCallback.cpp
@@ -379,14 +379,17 @@ void FNetworkModularVehicleInputs::DecayData(float DecayAmount)
 
 uint64 FNetworkModularVehicleStateNetTokenData::GetUniqueKey() const
 {
-	uint64 HashOfHashes = CityHash64((const char*)Hashes.GetData(),sizeof(uint32)*Hashes.Num());
-	uint64 HashOfIndexes = CityHash64((const char*)Indexes.GetData(),sizeof(int32)*Indexes.Num());
-	uint64 HashOfShouldSerialize = CityHash64((const char*)ModuleShouldSerialize.GetData(),sizeof(bool)*ModuleShouldSerialize.Num());
-	return HashOfHashes ^ HashOfIndexes ^ HashOfShouldSerialize;
+	uint64 HashOfHashes = GetTypeHash(Hashes);
+	uint64 HashOfIndexes = GetTypeHash(Indexes);
+	uint64 HashOfShouldSerialize = GetTypeHash(ModuleShouldSerialize);
+	return (HashOfHashes<<32) ^ HashOfIndexes ^ HashOfShouldSerialize;
 }
 
 void FNetworkModularVehicleStateNetTokenData::Init(const Chaos::FModuleNetDataArray& ModuleData)
 {
+	Hashes.Reset(ModuleData.Num());
+	Indexes.Reset(ModuleData.Num());
+	ModuleShouldSerialize.Reset(ModuleData.Num());
 	for (int32 Idx = 0; Idx < ModuleData.Num(); Idx++)
 	{
 		uint32 Hash = Chaos::FModuleFactoryRegister::GetModuleHash(ModuleData[Idx]->GetSimType());
diff --git a/Engine/Plugins/Experimental/ChaosModularVehicle/Source/ChaosModularVehicleEngine/Public/ChaosModularVehicle/ModuleInputTokenStore.h b/Engine/Plugins/Experimental/ChaosModularVehicle/Source/ChaosModularVehicleEngine/Public/ChaosModularVehicle/ModuleInputTokenStore.h
--- a/Engine/Plugins/Experimental/ChaosModularVehicle/Source/ChaosModularVehicleEngine/Public/ChaosModularVehicle/ModuleInputTokenStore.h
+++ b/Engine/Plugins/Experimental/ChaosModularVehicle/Source/ChaosModularVehicleEngine/Public/ChaosModularVehicle/ModuleInputTokenStore.h
@@ -21,13 +21,15 @@ struct FModuleInputNetTokenData
 	
 	uint64 GetUniqueKey() const
 	{
-		uint64 HashOfTypes = CityHash64((const char*)Types.GetData(),sizeof(uint8)*Types.Num());
-		uint64 HashOfDecayValues = CityHash64((const char*)DecayValues.GetData(),sizeof(bool)*DecayValues.Num());
-		return HashOfTypes ^ HashOfDecayValues;
+		uint64 HashOfTypes = GetTypeHash(Types);
+		uint64 HashOfDecayValues = GetTypeHash(DecayValues);
+		return (HashOfTypes<<32) ^ HashOfDecayValues;
 	}
 
 	void Init(const TArray<FModuleInputValue>& ModuleInputs)
 	{
+		Types.Reset(ModuleInputs.Num());
+		DecayValues.Reset(ModuleInputs.Num());
 		for (int32 Idx = 0; Idx < ModuleInputs.Num(); Idx++)
 		{
 			Types.Add(static_cast<uint8>(ModuleInputs[Idx].GetValueType()));
