Potential crash fix

diff --git a/Engine/Plugins/Experimental/ChaosModularVehicle/Source/ChaosModularVehicleEngine/Private/ChaosModularVehicle/ChaosSimModuleManager.cpp b/Engine/Plugins/Experimental/ChaosModularVehicle/Source/ChaosModularVehicleEngine/Private/ChaosModularVehicle/ChaosSimModuleManager.cpp
--- a/Engine/Plugins/Experimental/ChaosModularVehicle/Source/ChaosModularVehicleEngine/Private/ChaosModularVehicle/ChaosSimModuleManager.cpp
+++ b/Engine/Plugins/Experimental/ChaosModularVehicle/Source/ChaosModularVehicleEngine/Private/ChaosModularVehicle/ChaosSimModuleManager.cpp
@@ -81,9 +81,10 @@ void FChaosSimModuleManager::OnShowDebugInfo(AHUD* HUD, UCanvas* Canvas, const F
 		int ShowVehicleIndex = 0;
 		if (!Manager->CUVehicles.IsEmpty())
 		{
-			if (Manager->CUVehicles[ShowVehicleIndex].IsValid())
+			TStrongObjectPtr<UModularVehicleBaseComponent> StrongPtr = Manager->CUVehicles[ShowVehicleIndex].Pin();
+			if (StrongPtr.IsValid())
 			{
-				Manager->CUVehicles[ShowVehicleIndex]->ShowDebugInfo(HUD, Canvas, DisplayInfo, YL, YPos);
+				StrongPtr->ShowDebugInfo(HUD, Canvas, DisplayInfo, YL, YPos);
 			}
 		}
 	}
@@ -193,11 +194,12 @@ void FChaosSimModuleManager::RemoveVehicle(TWeakObjectPtr<UModularVehicleBaseCom
 
 void FChaosSimModuleManager::ScenePreTick(FPhysScene* PhysScene, float DeltaTime)
 {
-	for (int32 i = 0; i < CUVehicles.Num(); ++i)
+	for (TWeakObjectPtr<UModularVehicleBaseComponent>& Vehicle : CUVehicles)
 	{
-		if (CUVehicles[i] != nullptr)
+		TStrongObjectPtr<UModularVehicleBaseComponent> StrongPtr = Vehicle.Pin();
+		if (StrongPtr.IsValid())
 		{
-			CUVehicles[i]->PreTickGT(DeltaTime);
+			StrongPtr->PreTickGT(DeltaTime);
 		}
 	}
 }
@@ -216,10 +218,14 @@ void FChaosSimModuleManager::Update(FPhysScene* PhysScene, float DeltaTime)
 	{
 		FChaosSimModuleManagerAsyncInput* AsyncInput = AsyncCallback->GetProducerInputData_External();
 
-		for (TWeakObjectPtr<UModularVehicleBaseComponent> Vehicle : CUVehicles)
+		for (TWeakObjectPtr<UModularVehicleBaseComponent>& Vehicle : CUVehicles)
 		{
-			Vehicle->Update(DeltaTime);
-			Vehicle->FinalizeSimCallbackData(*AsyncInput);
+			TStrongObjectPtr<UModularVehicleBaseComponent> StrongPtr = Vehicle.Pin();
+			if (StrongPtr.IsValid())
+			{
+				StrongPtr->Update(DeltaTime);
+				StrongPtr->FinalizeSimCallbackData(*AsyncInput);
+			}
 		}
 	}
 }
@@ -227,11 +233,12 @@ void FChaosSimModuleManager::Update(FPhysScene* PhysScene, float DeltaTime)
 
 void FChaosSimModuleManager::PostUpdate(FChaosScene* PhysScene)
 {
-	for (TWeakObjectPtr<UModularVehicleBaseComponent> Vehicle : CUVehicles)
+	for (TWeakObjectPtr<UModularVehicleBaseComponent>& Vehicle : CUVehicles)
 	{
-		if (Vehicle.IsValid())
+		TStrongObjectPtr<UModularVehicleBaseComponent> StrongPtr = Vehicle.Pin();
+		if (StrongPtr.IsValid())
 		{
-			Vehicle->PostUpdate();
+			StrongPtr->PostUpdate();
 		}
 	}
 }
@@ -246,9 +253,13 @@ void FChaosSimModuleManager::InjectInputs_External(int32 PhysicsStep, int32 NumS
 	FChaosSimModuleManagerAsyncInput* AsyncInput = AsyncCallback->GetProducerInputData_External();
 	check(AsyncInput);
 
-	for (TWeakObjectPtr<UModularVehicleBaseComponent> Vehicle : CUVehicles)
+	for (TWeakObjectPtr<UModularVehicleBaseComponent>& Vehicle : CUVehicles)
 	{
-		Vehicle->ProduceInput(PhysicsStep, NumSteps);
+		TStrongObjectPtr<UModularVehicleBaseComponent> StrongPtr = Vehicle.Pin();
+		if (StrongPtr.IsValid())
+		{
+			StrongPtr->ProduceInput(PhysicsStep, NumSteps);
+		}
 	}
 
 }
@@ -318,20 +329,24 @@ void FChaosSimModuleManager::ParallelUpdateVehicles(float DeltaSeconds)
 	{
 		int32 NumVehiclesInActiveBatch = 0;
 
-		for (TWeakObjectPtr<UModularVehicleBaseComponent> Vehicle : CUVehicles)
+		for (TWeakObjectPtr<UModularVehicleBaseComponent>& Vehicle : CUVehicles)
 		{
-			auto NextOutput = PendingOutputs.Num() > 0 ? PendingOutputs[0].Get() : nullptr;
-			float Alpha = 0.f;
-			if (NextOutput && LatestOutput)
+			TStrongObjectPtr<UModularVehicleBaseComponent> StrongPtr = Vehicle.Pin();
+			if (StrongPtr.IsValid())
 			{
-				const float Denom = NextOutput->InternalTime - LatestOutput->InternalTime;
-				if (Denom > SMALL_NUMBER)
+				auto NextOutput = PendingOutputs.Num() > 0 ? PendingOutputs[0].Get() : nullptr;
+				float Alpha = 0.f;
+				if (NextOutput && LatestOutput)
 				{
-					Alpha = (ResultsTime - LatestOutput->InternalTime) / Denom;
+					const float Denom = NextOutput->InternalTime - LatestOutput->InternalTime;
+					if (Denom > SMALL_NUMBER)
+					{
+						Alpha = (ResultsTime - LatestOutput->InternalTime) / Denom;
+					}
 				}
-			}
 
-			AsyncInput->VehicleInputs.Add(Vehicle->SetCurrentAsyncData(AsyncInput->VehicleInputs.Num(), LatestOutput.Get(), NextOutput, Alpha, Timestamp));
+				AsyncInput->VehicleInputs.Add(StrongPtr->SetCurrentAsyncData(AsyncInput->VehicleInputs.Num(), LatestOutput.Get(), NextOutput, Alpha, Timestamp));
+			}
 		}
 
 	}
@@ -345,7 +360,11 @@ void FChaosSimModuleManager::ParallelUpdateVehicles(float DeltaSeconds)
 		auto LambdaParallelUpdate2 = [DeltaSeconds, &AwakeVehiclesBatch](int32 Idx)
 		{
 			TWeakObjectPtr<UModularVehicleBaseComponent> Vehicle = AwakeVehiclesBatch[Idx];
-			Vehicle->ParallelUpdate(DeltaSeconds); // gets output state from PT
+			TStrongObjectPtr<UModularVehicleBaseComponent> StrongPtr = Vehicle.Pin();
+			if (StrongPtr.IsValid())
+			{
+				StrongPtr->ParallelUpdate(DeltaSeconds); // gets output state from PT
+			}
 		};
 
 		ParallelFor(AwakeVehiclesBatch.Num(), LambdaParallelUpdate2, ForceSingleThread);
