[Chaos Modular Vehicles]

If FNetworkModularVehicleStates::InterpolateData has mismatched number of module datas between min and max, early out rather than crashing as a workaround.

diff --git a/Engine/Plugins/Experimental/ChaosModularVehicle/Source/ChaosModularVehicleEngine/Private/ChaosModularVehicle/ChaosSimModuleManagerAsyncCallback.cpp b/Engine/Plugins/Experimental/ChaosModularVehicle/Source/ChaosModularVehicleEngine/Private/ChaosModularVehicle/ChaosSimModuleManagerAsyncCallback.cpp
--- a/Engine/Plugins/Experimental/ChaosModularVehicle/Source/ChaosModularVehicleEngine/Private/ChaosModularVehicle/ChaosSimModuleManagerAsyncCallback.cpp
+++ b/Engine/Plugins/Experimental/ChaosModularVehicle/Source/ChaosModularVehicleEngine/Private/ChaosModularVehicle/ChaosSimModuleManagerAsyncCallback.cpp
@@ -643,6 +643,11 @@ void FNetworkModularVehicleStates::InterpolateData(const FNetworkPhysicsData& Mi
 
 	const float LerpFactor = (LocalFrame - MinState.LocalFrame) / (MaxState.LocalFrame - MinState.LocalFrame);
 
+	if (ModuleData.Num() != MinState.ModuleData.Num() || ModuleData.Num() != MaxState.ModuleData.Num())
+	{
+		UE_LOG(LogModularVehicleSim, Error, TEXT("Mismatch in module data num when interpolating between min and max states!"));
+		return;
+	}
 	for (int I = 0; I < ModuleData.Num(); I++)
 	{
 		// if these don't match then something has gone terribly wrong
