Setting a start frame for test inputs as part of the input producer

- Added a virtual function to UVehicleInputProducerBase to specify which frame should be the cue for tests to start. Defaults to -1 which means the test inputs will start immediately.
- Changed start frame logic so that the start frame will only be overridden to start on the current frame if the specified frame is less than 0. Previously if the specified frame was in the future, it would get also reset to the current frame. Now it will correctly wait until the specified frame is current or in the past before executing test inputs.

diff --git a/Engine/Plugins/Experimental/ChaosModularVehicle/Source/ChaosModularVehicleEngine/Private/ChaosModularVehicle/ModularVehicleBaseComponent.cpp b/Engine/Plugins/Experimental/ChaosModularVehicle/Source/ChaosModularVehicleEngine/Private/ChaosModularVehicle/ModularVehicleBaseComponent.cpp
--- a/Engine/Plugins/Experimental/ChaosModularVehicle/Source/ChaosModularVehicleEngine/Private/ChaosModularVehicle/ModularVehicleBaseComponent.cpp
+++ b/Engine/Plugins/Experimental/ChaosModularVehicle/Source/ChaosModularVehicleEngine/Private/ChaosModularVehicle/ModularVehicleBaseComponent.cpp
@@ -661,7 +661,7 @@ void UModularVehicleBaseComponent::SetupInputConfiguration()
 		if (InputProducer->GetTestInputBuffer())
 		{
 			// deliberate deep copy rather than copy pointer
-			VehicleSimulationPT->SetTestInputBuffer(*InputProducer->GetTestInputBuffer(), InputProducer->IsLoopingTestInputBuffer());
+			VehicleSimulationPT->SetTestInputBuffer(*InputProducer->GetTestInputBuffer(), InputProducer->IsLoopingTestInputBuffer(), InputProducer->GetTestInputStartFrame());
 		}
 	}
 
diff --git a/Engine/Plugins/Experimental/ChaosModularVehicle/Source/ChaosModularVehicleEngine/Private/ChaosModularVehicle/ModularVehicleSimulationCU.cpp b/Engine/Plugins/Experimental/ChaosModularVehicle/Source/ChaosModularVehicleEngine/Private/ChaosModularVehicle/ModularVehicleSimulationCU.cpp
--- a/Engine/Plugins/Experimental/ChaosModularVehicle/Source/ChaosModularVehicleEngine/Private/ChaosModularVehicle/ModularVehicleSimulationCU.cpp
+++ b/Engine/Plugins/Experimental/ChaosModularVehicle/Source/ChaosModularVehicleEngine/Private/ChaosModularVehicle/ModularVehicleSimulationCU.cpp
@@ -110,21 +110,25 @@ void FModularVehicleSimulation::SimulateModuleTree(UWorld* InWorld, float DeltaS
 			check(RigidsSolver);
 			const int32 CurrentPhysicsFrame = RigidsSolver->GetCurrentFrame();
 
-			if (TestInputBufferStartFrame < 0 || TestInputBufferStartFrame > CurrentPhysicsFrame)
+			if (TestInputBufferStartFrame < 0)
 			{
 				TestInputBufferStartFrame = CurrentPhysicsFrame;
 			}
-			int32 InputFrame = (CurrentPhysicsFrame - TestInputBufferStartFrame);
 
-			if (ImplementsLoopingTestBuffer() && !TestInputBuffer.IsValidIndex(InputFrame))
+			if (TestInputBufferStartFrame <= CurrentPhysicsFrame)
 			{
-				TestInputBufferStartFrame = CurrentPhysicsFrame;
-				InputFrame = 0;
-			}
+				int32 InputFrame = (CurrentPhysicsFrame - TestInputBufferStartFrame);
 
-			if (TestInputBuffer.IsValidIndex(InputFrame))
-			{
-				Container = TestInputBuffer[InputFrame];
+				if (ImplementsLoopingTestBuffer() && !TestInputBuffer.IsValidIndex(InputFrame))
+				{
+					TestInputBufferStartFrame = CurrentPhysicsFrame;
+					InputFrame = 0;
+				}
+
+				if (TestInputBuffer.IsValidIndex(InputFrame))
+				{
+					Container = TestInputBuffer[InputFrame];
+				}
 			}
 		}
 
diff --git a/Engine/Plugins/Experimental/ChaosModularVehicle/Source/ChaosModularVehicleEngine/Public/ChaosModularVehicle/ModularVehicleSimulationCU.h b/Engine/Plugins/Experimental/ChaosModularVehicle/Source/ChaosModularVehicleEngine/Public/ChaosModularVehicle/ModularVehicleSimulationCU.h
--- a/Engine/Plugins/Experimental/ChaosModularVehicle/Source/ChaosModularVehicleEngine/Public/ChaosModularVehicle/ModularVehicleSimulationCU.h
+++ b/Engine/Plugins/Experimental/ChaosModularVehicle/Source/ChaosModularVehicleEngine/Public/ChaosModularVehicle/ModularVehicleSimulationCU.h
@@ -79,12 +79,12 @@ public:
 		StateNameMap = InNameMap;
 	}
 
-	void SetTestInputBuffer(TArray<FModuleInputContainer>& InTestInputBuffer, bool bInIsLoopBuffer)
+	void SetTestInputBuffer(TArray<FModuleInputContainer>& InTestInputBuffer, bool bInIsLoopBuffer, int32 InTestInputBufferStartFrame)
 	{
 		UE::TWriteScopeLock InputConfigLock(InputConfigurationLock);
 		bIsLoopingTestInputBuffer = bInIsLoopBuffer;
 		TestInputBuffer = InTestInputBuffer;
-		TestInputBufferStartFrame = -1;
+		TestInputBufferStartFrame = InTestInputBufferStartFrame;
 	}
 
 	UE_API void AppendTreeUpdates(Chaos::FSimTreeUpdates* InNextTreeUpdatesInternal);
diff --git a/Engine/Source/Runtime/Experimental/ChaosVehicles/ChaosVehiclesCore/Public/SimModule/ModuleInput.h b/Engine/Source/Runtime/Experimental/ChaosVehicles/ChaosVehiclesCore/Public/SimModule/ModuleInput.h
--- a/Engine/Source/Runtime/Experimental/ChaosVehicles/ChaosVehiclesCore/Public/SimModule/ModuleInput.h
+++ b/Engine/Source/Runtime/Experimental/ChaosVehicles/ChaosVehiclesCore/Public/SimModule/ModuleInput.h
@@ -776,6 +776,9 @@ public:
 	/** Special case override for providing test input straight onto the physics thread */
 	virtual bool IsLoopingTestInputBuffer() { return false; }
 
+	/** Special case override for providing test input straight onto the physics thread */
+	virtual int32 GetTestInputStartFrame() { return -1; }
+
 	EModuleInputQuantizationType InputQuantizationType;
 };
 
