Potential fix for rare crash which was happening when a vehicle was being destroyed.

diff --git a/Engine/Plugins/Experimental/ChaosModularVehicle/Source/ChaosModularVehicleEngine/Private/ChaosModularVehicle/ModularVehicleBaseComponent.cpp b/Engine/Plugins/Experimental/ChaosModularVehicle/Source/ChaosModularVehicleEngine/Private/ChaosModularVehicle/ModularVehicleBaseComponent.cpp
--- a/Engine/Plugins/Experimental/ChaosModularVehicle/Source/ChaosModularVehicleEngine/Private/ChaosModularVehicle/ModularVehicleBaseComponent.cpp
+++ b/Engine/Plugins/Experimental/ChaosModularVehicle/Source/ChaosModularVehicleEngine/Private/ChaosModularVehicle/ModularVehicleBaseComponent.cpp
@@ -1424,17 +1424,21 @@ void UModularVehicleBaseComponent::DestroyVehicleSim()
 				});
 
 		}
+	}
 
-		if (bUsingNetworkPhysicsPrediction && NetworkPhysicsComponent)
-		{
-			NetworkPhysicsComponent->RemoveDataHistory();
-		}
+}
 
+void UModularVehicleBaseComponent::UninitializeComponent()
+{
+	if (bUsingNetworkPhysicsPrediction && NetworkPhysicsComponent)
+	{
+		NetworkPhysicsComponent->DestroyComponent();
+		NetworkPhysicsComponent = nullptr;
 	}
 
+	Super::UninitializeComponent();
 }
 
-
 void UModularVehicleBaseComponent::AddComponentToCluster(USceneComponent* Component)
 {
 	TArray<int32> BoneIds;
diff --git a/Engine/Plugins/Experimental/ChaosModularVehicle/Source/ChaosModularVehicleEngine/Public/ChaosModularVehicle/ModularVehicleBaseComponent.h b/Engine/Plugins/Experimental/ChaosModularVehicle/Source/ChaosModularVehicleEngine/Public/ChaosModularVehicle/ModularVehicleBaseComponent.h
--- a/Engine/Plugins/Experimental/ChaosModularVehicle/Source/ChaosModularVehicleEngine/Public/ChaosModularVehicle/ModularVehicleBaseComponent.h
+++ b/Engine/Plugins/Experimental/ChaosModularVehicle/Source/ChaosModularVehicleEngine/Public/ChaosModularVehicle/ModularVehicleBaseComponent.h
@@ -137,6 +137,7 @@ public:
 	virtual bool ShouldCreatePhysicsState() const override { return true; }
 	UE_API virtual void OnCreatePhysicsState() override;
 	UE_API virtual void OnDestroyPhysicsState() override;
+	UE_API virtual void UninitializeComponent() override;
 	UE_API virtual void SetClusterComponent(UClusterUnionComponent* InPhysicalComponent);
 
 	UE_API virtual void BeginPlay() override;
