Fix Predictive interpolation networking mode for modular vehicles

diff --git a/Engine/Plugins/Experimental/ChaosModularVehicle/Source/ChaosModularVehicleEngine/Private/ChaosModularVehicle/ModularVehicleBaseComponent.cpp b/Engine/Plugins/Experimental/ChaosModularVehicle/Source/ChaosModularVehicleEngine/Private/ChaosModularVehicle/ModularVehicleBaseComponent.cpp
--- a/Engine/Plugins/Experimental/ChaosModularVehicle/Source/ChaosModularVehicleEngine/Private/ChaosModularVehicle/ModularVehicleBaseComponent.cpp
+++ b/Engine/Plugins/Experimental/ChaosModularVehicle/Source/ChaosModularVehicleEngine/Private/ChaosModularVehicle/ModularVehicleBaseComponent.cpp
@@ -769,8 +769,13 @@ void UModularVehicleBaseComponent::UpdateState(float DeltaTime)
 	{
 		if (!bUsingNetworkPhysicsPrediction)
 		{
-			// and send to server - (ServerUpdateState_Implementation below)
-			ServerUpdateState(InputsContainer, bKeepVehicleAwake);
+			// Produce input but only for remote server (ServerInputsContainer), not the local client (InputsContainer)
+			if (InputProducer)
+			{
+				FModuleInputContainer ServerInputsContainer;
+				InputProducer->ProduceInput(0, 0, InputNameMap, ServerInputsContainer);
+				ServerUpdateState(ServerInputsContainer, bKeepVehicleAwake);
+			}
 		}
 
 		if (PawnOwner && PawnOwner->IsNetMode(NM_Client))
@@ -778,9 +783,10 @@ void UModularVehicleBaseComponent::UpdateState(float DeltaTime)
 			MarkForClientCameraUpdate();
 		}
 	}
-	else if (!bUsingNetworkPhysicsPrediction)
+	
+	if (!bUsingNetworkPhysicsPrediction)
 	{
-		// use replicated values for remote pawns
+		// use the replicated values for all pawns Autonomous and Simulated
 		InputsContainer = ReplicatedState.Container;
 		bKeepVehicleAwake = ReplicatedState.KeepAwake;
 	}
@@ -793,9 +799,14 @@ bool UModularVehicleBaseComponent::ServerUpdateState_Validate(const FModuleInput
 
 void UModularVehicleBaseComponent::ServerUpdateState_Implementation(const FModuleInputContainer& InputsIn, bool KeepAwake)
 {
-	// update state of inputs
+	// update state of inputs that will be replicated back to the clients
 	ReplicatedState.KeepAwake = KeepAwake;
 	ReplicatedState.Container = InputsIn;
+
+	MARK_PROPERTY_DIRTY_FROM_NAME(UModularVehicleBaseComponent, ReplicatedState, this);
+
+	// apply control inputs on server as received from the original client
+	InputsContainer = InputsIn;
 }
 
 
