[NetPhysics] Optimization to not cache physics result during resimulation frames, except for the last frame.
- Also don't populate simulation state in AsyncOutput during resimulation for chaos modular vehicles

diff --git a/Engine/Plugins/Experimental/ChaosModularVehicle/Source/ChaosModularVehicleEngine/Private/ChaosModularVehicle/ChaosSimModuleManagerAsyncCallback.cpp b/Engine/Plugins/Experimental/ChaosModularVehicle/Source/ChaosModularVehicleEngine/Private/ChaosModularVehicle/ChaosSimModuleManagerAsyncCallback.cpp
--- a/Engine/Plugins/Experimental/ChaosModularVehicle/Source/ChaosModularVehicleEngine/Private/ChaosModularVehicle/ChaosSimModuleManagerAsyncCallback.cpp
+++ b/Engine/Plugins/Experimental/ChaosModularVehicle/Source/ChaosModularVehicleEngine/Private/ChaosModularVehicle/ChaosSimModuleManagerAsyncCallback.cpp
@@ -102,6 +102,7 @@ void FChaosSimModuleManagerAsyncCallback::OnPreSimulate_Internal()
 	// beware running the vehicle simulation in parallel, code must remain threadsafe
 	auto LambdaParallelUpdate = [World, DeltaTime, SimTime, &InputVehiclesBatch, &OutputVehiclesBatch](int32 Idx)
 	{
+		QUICK_SCOPE_CYCLE_COUNTER(Stat_LambdaParallelUpdate);
 		const FModularVehicleAsyncInput& VehicleInput = *InputVehiclesBatch[Idx];
 
 		if (VehicleInput.Proxy == nullptr)
@@ -184,6 +185,7 @@ void FChaosSimModuleManagerAsyncCallback::OnContactModification_Internal(Chaos::
 
 TUniquePtr<FModularVehicleAsyncOutput> FModularVehicleAsyncInput::Simulate(UWorld* World, const float DeltaSeconds, const float TotalSeconds, bool& bWakeOut) const
 {
+	QUICK_SCOPE_CYCLE_COUNTER(Stat_FModularVehicleAsyncInput_Simulate);
 	TUniquePtr<FModularVehicleAsyncOutput> Output = MakeUnique<FModularVehicleAsyncOutput>();
 
 	//support nullptr because it allows us to go wide on filling the async inputs
@@ -199,8 +201,20 @@ TUniquePtr<FModularVehicleAsyncOutput> FModularVehicleAsyncInput::Simulate(UWorl
 			// FILL OUTPUT DATA HERE THAT WILL GET PASSED BACK TO THE GAME THREAD
 			Sim->Simulate(World, DeltaSeconds, *this, *Output.Get(), Proxy);
 
-			FModularVehicleAsyncOutput& OutputData = *Output.Get();
-			Sim->FillOutputState(OutputData);
+			Chaos::FPhysicsSolver* Solver = Proxy->GetSolver<Chaos::FPhysicsSolver>();
+			if (Solver && Solver->GetEvolution())
+			{
+				if (!Solver->GetEvolution()->IsResimming())
+				{
+					FModularVehicleAsyncOutput& OutputData = *Output.Get();
+					Sim->FillOutputState(OutputData);
+				}
+			}
+			else
+			{
+				FModularVehicleAsyncOutput& OutputData = *Output.Get();
+				Sim->FillOutputState(OutputData);
+			}
 		}
 	}
 
diff --git a/Engine/Source/Runtime/Experimental/Chaos/Private/PBDRigidsSolver.cpp b/Engine/Source/Runtime/Experimental/Chaos/Private/PBDRigidsSolver.cpp
--- a/Engine/Source/Runtime/Experimental/Chaos/Private/PBDRigidsSolver.cpp
+++ b/Engine/Source/Runtime/Experimental/Chaos/Private/PBDRigidsSolver.cpp
@@ -2309,10 +2309,12 @@ namespace Chaos
 	{
 		LLM_SCOPE(ELLMTag::Chaos);
 		SCOPE_CYCLE_COUNTER(STAT_BufferPhysicsResults);
-
-		EventPreBuffer.Broadcast(MLastDt);
-		GetDirtyParticlesBuffer()->CaptureSolverData(this);
-		BufferPhysicsResults();
+		if (!IsResimming() || (GetRewindData() && GetRewindData()->IsFinalResim()))
+		{
+			EventPreBuffer.Broadcast(MLastDt);
+			GetDirtyParticlesBuffer()->CaptureSolverData(this);
+			BufferPhysicsResults();
+		}
 	}
 
 
