Fix bug where module creations events are being dropped when GT running slow & physics runs multiple ticks

diff --git a/Engine/Plugins/Experimental/ChaosModularVehicle/Source/ChaosModularVehicleEngine/Private/ChaosModularVehicle/ChaosSimModuleManager.cpp b/Engine/Plugins/Experimental/ChaosModularVehicle/Source/ChaosModularVehicleEngine/Private/ChaosModularVehicle/ChaosSimModuleManager.cpp
--- a/Engine/Plugins/Experimental/ChaosModularVehicle/Source/ChaosModularVehicleEngine/Private/ChaosModularVehicle/ChaosSimModuleManager.cpp
+++ b/Engine/Plugins/Experimental/ChaosModularVehicle/Source/ChaosModularVehicleEngine/Private/ChaosModularVehicle/ChaosSimModuleManager.cpp
@@ -76,15 +76,16 @@ void FChaosSimModuleManager::OnShowDebugInfo(AHUD* HUD, UCanvas* Canvas, const F
 	static const FName NAME_ModularVehicle("ModularVehicle");
 	if (Canvas && HUD->ShouldDisplayDebug(NAME_ModularVehicle))
 	{
-		FChaosSimModuleManager* Manager = FChaosSimModuleManager::GetManagerFromScene(HUD->GetWorld()->GetPhysicsScene());
-
-		int ShowVehicleIndex = 0;
-		if (!Manager->CUVehicles.IsEmpty())
+		if (FChaosSimModuleManager* Manager = FChaosSimModuleManager::GetManagerFromScene(HUD->GetWorld()->GetPhysicsScene()))
 		{
-			TStrongObjectPtr<UModularVehicleBaseComponent> StrongPtr = Manager->CUVehicles[ShowVehicleIndex].Pin();
-			if (StrongPtr.IsValid())
+			int ShowVehicleIndex = 0;
+			if (!Manager->CUVehicles.IsEmpty())
 			{
-				StrongPtr->ShowDebugInfo(HUD, Canvas, DisplayInfo, YL, YPos);
+				TStrongObjectPtr<UModularVehicleBaseComponent> StrongPtr = Manager->CUVehicles[ShowVehicleIndex].Pin();
+				if (StrongPtr.IsValid())
+				{
+					StrongPtr->ShowDebugInfo(HUD, Canvas, DisplayInfo, YL, YPos);
+				}
 			}
 		}
 	}
@@ -179,7 +180,6 @@ FChaosSimModuleManager* FChaosSimModuleManager::GetManagerFromScene(FPhysScene*
 void FChaosSimModuleManager::AddVehicle(TWeakObjectPtr<UModularVehicleBaseComponent> Vehicle)
 {
 	check(Vehicle != NULL);
-	check(Vehicle->PhysicsVehicleOutput());
 	check(AsyncCallback);
 	CUVehicles.Add(Vehicle);
 }
@@ -274,7 +274,7 @@ void FChaosSimModuleManager::ParallelUpdateVehicles(float DeltaSeconds)
 
 	{
 		// We pass pointers from TArray so this reserve is critical. Otherwise realloc happens
-		AsyncInput->VehicleInputs.Reserve(CUVehicles.Num() + GCVehicles.Num());
+		AsyncInput->VehicleInputs.Reserve(CUVehicles.Num());
 		AsyncInput->Timestamp = Timestamp;
 		AsyncInput->World = Scene.GetOwningWorld();
 	}
@@ -288,6 +288,27 @@ void FChaosSimModuleManager::ParallelUpdateVehicles(float DeltaSeconds)
 		}
 	}
 
+	// ensure we don't miss any module creation events
+	TArray<Chaos::FCreatedModules> CombinedNewlyCreatedModuleGuids;
+	Chaos::FCreatedModules ModulesEventData;
+	CombinedNewlyCreatedModuleGuids.Init(ModulesEventData, CUVehicles.Num());
+
+	for (int32 OutIdx = 0; OutIdx < PendingOutputs.Num(); ++OutIdx)
+	{
+		if (FChaosSimModuleManagerAsyncOutput* AsyncOutput = PendingOutputs[OutIdx].Get())
+		{
+			for (int32 VehicleIdx = 0; VehicleIdx < AsyncOutput->VehicleOutputs.Num(); ++VehicleIdx)
+			{
+				TArray<Chaos::FCreatedModule>& NewlyCreatedModuleGuids = AsyncOutput->VehicleOutputs[VehicleIdx]->VehicleSimOutput.NewlyCreatedModuleGuids;
+
+				if (NewlyCreatedModuleGuids.Num() > 0 && VehicleIdx < CombinedNewlyCreatedModuleGuids.Num())
+				{
+					CombinedNewlyCreatedModuleGuids[VehicleIdx].ModuleEvents.Append(NewlyCreatedModuleGuids);
+				}
+			}
+		}
+	}
+
 	// Since we are in pre-physics, delta seconds is not accounted for in external time yet
 	const double ResultsTime = AsyncCallback->GetSolver()->GetPhysicsResultsTime_External();
 
@@ -356,18 +377,18 @@ void FChaosSimModuleManager::ParallelUpdateVehicles(float DeltaSeconds)
 	bool ForceSingleThread = !GSimModuleDebugParams.EnableMultithreading;
 
 	{
-		const auto& AwakeVehiclesBatch = CUVehicles;
-		auto LambdaParallelUpdate2 = [DeltaSeconds, &AwakeVehiclesBatch](int32 Idx)
-		{
-			TWeakObjectPtr<UModularVehicleBaseComponent> Vehicle = AwakeVehiclesBatch[Idx];
-			TStrongObjectPtr<UModularVehicleBaseComponent> StrongPtr = Vehicle.Pin();
-			if (StrongPtr.IsValid())
+		const TArray<TWeakObjectPtr<UModularVehicleBaseComponent>>& AwakeVehiclesBatch = CUVehicles;
+		const TArray<Chaos::FCreatedModules>& ModuleEvents = CombinedNewlyCreatedModuleGuids;
+		auto LambdaParallelUpdate2 = [DeltaSeconds, &AwakeVehiclesBatch, &ModuleEvents](int32 Idx)
 			{
-				StrongPtr->ParallelUpdate(DeltaSeconds); // gets output state from PT
-			}
-		};
+				TWeakObjectPtr<UModularVehicleBaseComponent> Vehicle = AwakeVehiclesBatch[Idx];
+				TStrongObjectPtr<UModularVehicleBaseComponent> StrongPtr = Vehicle.Pin();
+				if (StrongPtr.IsValid())
+				{
+					StrongPtr->ParallelUpdate(DeltaSeconds, ModuleEvents[Idx]); // gets output state from PT
+				}
+			};
 
 		ParallelFor(AwakeVehiclesBatch.Num(), LambdaParallelUpdate2, ForceSingleThread);
 	}
-
 }
\ No newline at end of file
diff --git a/Engine/Plugins/Experimental/ChaosModularVehicle/Source/ChaosModularVehicleEngine/Private/ChaosModularVehicle/ModularVehicleBaseComponent.cpp b/Engine/Plugins/Experimental/ChaosModularVehicle/Source/ChaosModularVehicleEngine/Private/ChaosModularVehicle/ModularVehicleBaseComponent.cpp
--- a/Engine/Plugins/Experimental/ChaosModularVehicle/Source/ChaosModularVehicleEngine/Private/ChaosModularVehicle/ModularVehicleBaseComponent.cpp
+++ b/Engine/Plugins/Experimental/ChaosModularVehicle/Source/ChaosModularVehicleEngine/Private/ChaosModularVehicle/ModularVehicleBaseComponent.cpp
@@ -940,7 +940,7 @@ void UModularVehicleBaseComponent::SetClusterComponent(UClusterUnionComponent* I
 /***************************************************************************/
 /* READ OUTPUT DATA - Access the async output data from the Physics Thread */
 /***************************************************************************/
-void UModularVehicleBaseComponent::ParallelUpdate(float DeltaTime)
+void UModularVehicleBaseComponent::ParallelUpdate(const float DeltaTime, const Chaos::FCreatedModules& ModuleEvents)
 {
 	if (FModularVehicleAsyncOutput* CurrentOutput = static_cast<FModularVehicleAsyncOutput*>(CurAsyncOutput))
 	{
@@ -948,7 +948,7 @@ void UModularVehicleBaseComponent::ParallelUpdate(float DeltaTime)
 		{
 			PVehicleOutput->Clean();
 
-			PVehicleOutput->NewlyCreatedModuleGuids = CurrentOutput->VehicleSimOutput.NewlyCreatedModuleGuids;
+			PVehicleOutput->NewlyCreatedModuleGuids = ModuleEvents.ModuleEvents;
 
 			int NumItems = CurrentOutput->VehicleSimOutput.SimTreeOutputData.Num();
 			PVehicleOutput->SimTreeOutputData.Reserve(NumItems);
@@ -1098,10 +1098,13 @@ void UModularVehicleBaseComponent::PostUpdate()
 {
 	ensure(PVehicleOutput);
 
-	// broadcast events
-	for (Chaos::FCreatedModules AddedModule : PVehicleOutput->NewlyCreatedModuleGuids)
+	if (PVehicleOutput)
 	{
-		BroadcastModuleAddedEvent(AddedModule.SimType, AddedModule.Guid, AddedModule.TreeIndex);
+		// broadcast events
+		for (Chaos::FCreatedModule AddedModule : PVehicleOutput->NewlyCreatedModuleGuids)
+		{
+			BroadcastModuleAddedEvent(AddedModule.SimType, AddedModule.Guid, AddedModule.TreeIndex);
+		}
 	}
 
 	// Update replication structures
diff --git a/Engine/Plugins/Experimental/ChaosModularVehicle/Source/ChaosModularVehicleEngine/Private/ChaosModularVehicle/ModularVehicleSimulationCU.cpp b/Engine/Plugins/Experimental/ChaosModularVehicle/Source/ChaosModularVehicleEngine/Private/ChaosModularVehicle/ModularVehicleSimulationCU.cpp
--- a/Engine/Plugins/Experimental/ChaosModularVehicle/Source/ChaosModularVehicleEngine/Private/ChaosModularVehicle/ModularVehicleSimulationCU.cpp
+++ b/Engine/Plugins/Experimental/ChaosModularVehicle/Source/ChaosModularVehicleEngine/Private/ChaosModularVehicle/ModularVehicleSimulationCU.cpp
@@ -498,7 +498,7 @@ void FModularVehicleSimulation::ActionTreeUpdates()
 			{
 				if (ModuleAdd.NewSimModule)
 				{
-					NewlyCreatedModuleGuids.Add(Chaos::FCreatedModules(
+					NewlyCreatedModuleGuids.Add(Chaos::FCreatedModule(
 						ModuleAdd.NewSimModule->GetSimType(),
 						ModuleAdd.NewSimModule->GetGuid(),
 						ModuleAdd.NewSimModule->GetTreeIndex()));
diff --git a/Engine/Plugins/Experimental/ChaosModularVehicle/Source/ChaosModularVehicleEngine/Public/ChaosModularVehicle/ChaosSimModuleManager.h b/Engine/Plugins/Experimental/ChaosModularVehicle/Source/ChaosModularVehicleEngine/Public/ChaosModularVehicle/ChaosSimModuleManager.h
--- a/Engine/Plugins/Experimental/ChaosModularVehicle/Source/ChaosModularVehicleEngine/Public/ChaosModularVehicle/ChaosSimModuleManager.h
+++ b/Engine/Plugins/Experimental/ChaosModularVehicle/Source/ChaosModularVehicleEngine/Public/ChaosModularVehicle/ChaosSimModuleManager.h
@@ -83,9 +83,6 @@ private:
 
 	static UE_API bool GInitialized;
 
-	// All instanced vehicles
-	TArray<TWeakObjectPtr<UModularVehicleComponent>> GCVehicles;
-
 	// Test out new vehicle bass class
 	TArray<TWeakObjectPtr<UModularVehicleBaseComponent>> CUVehicles;
 
diff --git a/Engine/Plugins/Experimental/ChaosModularVehicle/Source/ChaosModularVehicleEngine/Public/ChaosModularVehicle/ChaosSimModuleManagerAsyncCallback.h b/Engine/Plugins/Experimental/ChaosModularVehicle/Source/ChaosModularVehicleEngine/Public/ChaosModularVehicle/ChaosSimModuleManagerAsyncCallback.h
--- a/Engine/Plugins/Experimental/ChaosModularVehicle/Source/ChaosModularVehicleEngine/Public/ChaosModularVehicle/ChaosSimModuleManagerAsyncCallback.h
+++ b/Engine/Plugins/Experimental/ChaosModularVehicle/Source/ChaosModularVehicleEngine/Public/ChaosModularVehicle/ChaosSimModuleManagerAsyncCallback.h
@@ -205,7 +205,7 @@ struct FPhysicsVehicleOutput
 	}
 
 	TArray<Chaos::FSimOutputData*> SimTreeOutputData;
-	TArray<Chaos::FCreatedModules> NewlyCreatedModuleGuids;
+	TArray<Chaos::FCreatedModule> NewlyCreatedModuleGuids;
 };
 
 struct FPhysicsModularVehicleTraits
diff --git a/Engine/Plugins/Experimental/ChaosModularVehicle/Source/ChaosModularVehicleEngine/Public/ChaosModularVehicle/ModularVehicleBaseComponent.h b/Engine/Plugins/Experimental/ChaosModularVehicle/Source/ChaosModularVehicleEngine/Public/ChaosModularVehicle/ModularVehicleBaseComponent.h
--- a/Engine/Plugins/Experimental/ChaosModularVehicle/Source/ChaosModularVehicleEngine/Public/ChaosModularVehicle/ModularVehicleBaseComponent.h
+++ b/Engine/Plugins/Experimental/ChaosModularVehicle/Source/ChaosModularVehicleEngine/Public/ChaosModularVehicle/ModularVehicleBaseComponent.h
@@ -154,7 +154,7 @@ public:
 	UE_API void UpdateState(float DeltaTime);
 	UE_API TUniquePtr<FModularVehicleAsyncInput> SetCurrentAsyncData(int32 InputIdx, FChaosSimModuleManagerAsyncOutput* CurOutput, FChaosSimModuleManagerAsyncOutput* NextOutput, float Alpha, int32 VehicleManagerTimestamp);
 
-	UE_API void ParallelUpdate(float DeltaTime);
+	UE_API void ParallelUpdate(const float DeltaTime, const Chaos::FCreatedModules& ModuleEvents);
 	UE_API void Update(float DeltaTime);
 	UE_API void PostUpdate();
 	UE_API void FinalizeSimCallbackData(FChaosSimModuleManagerAsyncInput& Input);
diff --git a/Engine/Plugins/Experimental/ChaosModularVehicle/Source/ChaosModularVehicleEngine/Public/ChaosModularVehicle/ModularVehicleSimulationCU.h b/Engine/Plugins/Experimental/ChaosModularVehicle/Source/ChaosModularVehicleEngine/Public/ChaosModularVehicle/ModularVehicleSimulationCU.h
--- a/Engine/Plugins/Experimental/ChaosModularVehicle/Source/ChaosModularVehicleEngine/Public/ChaosModularVehicle/ModularVehicleSimulationCU.h
+++ b/Engine/Plugins/Experimental/ChaosModularVehicle/Source/ChaosModularVehicleEngine/Public/ChaosModularVehicle/ModularVehicleSimulationCU.h
@@ -134,7 +134,7 @@ public:
 	bool ImplementsLoopingTestBuffer() { return bIsLoopingTestInputBuffer; }
 	TArray<FModuleInputContainer> TestInputBuffer;
 	TArray<Chaos::FSimTreeUpdates> NextTreeUpdatesInternal;
-	TArray<Chaos::FCreatedModules> NewlyCreatedModuleGuids;
+	TArray<Chaos::FCreatedModule> NewlyCreatedModuleGuids;
 
 	Chaos::FPBDRigidParticleHandle* RootParticle = nullptr; // cached root particle
 	
diff --git a/Engine/Source/Runtime/Experimental/ChaosVehicles/ChaosVehiclesCore/Public/SimModule/SimulationModuleBase.h b/Engine/Source/Runtime/Experimental/ChaosVehicles/ChaosVehiclesCore/Public/SimModule/SimulationModuleBase.h
--- a/Engine/Source/Runtime/Experimental/ChaosVehicles/ChaosVehiclesCore/Public/SimModule/SimulationModuleBase.h
+++ b/Engine/Source/Runtime/Experimental/ChaosVehicles/ChaosVehiclesCore/Public/SimModule/SimulationModuleBase.h
@@ -73,18 +73,41 @@ namespace Chaos
 		bool bKeepVehicleAwake = false;
 	};
 
-	struct FCreatedModules
+	struct FCreatedModule
 	{
-		FCreatedModules(const FName InSimType, const int InGuid, const int InTreeIndex)
-			: SimType(InSimType)
-			, Guid(InGuid)
-			, TreeIndex(InTreeIndex)
+		FCreatedModule() : SimType("Empty"), Guid(INDEX_NONE), TreeIndex(INDEX_NONE)
+		{
+		}
+
+		FCreatedModule(const FName& InSimType, int InGuid, int InTreeIndex)
+			: SimType(InSimType), Guid(InGuid), TreeIndex(InTreeIndex)
+		{
+		}
+
+		FCreatedModule(const FCreatedModule& Other)
+			: SimType(Other.SimType), Guid(Other.Guid), TreeIndex(Other.TreeIndex)
+		{
+		}
+		FCreatedModule& operator=(const FCreatedModule& Other)
 		{
+			SimType = Other.SimType;
+			Guid = Other.Guid;
+			TreeIndex = Other.TreeIndex;
+			return *this;
 		}
 
-		const FName SimType;
+		FName SimType;
 		int Guid = INDEX_NONE;
-		int TreeIndex = INDEX_NONE;
+		int32 TreeIndex = INDEX_NONE;
+	};
+
+	struct FCreatedModules
+	{
+		FCreatedModules()
+		{
+		}
+
+		TArray<Chaos::FCreatedModule> ModuleEvents;
 	};
 
 	struct FSimModuleAnimationData
