[Chaos]
Add params to modular vehicle custom simulation callback API

diff --git a/Engine/Plugins/Experimental/ChaosModularVehicle/Source/ChaosModularVehicleEngine/Private/ChaosModularVehicle/ModularVehicleSimulationCU.cpp b/Engine/Plugins/Experimental/ChaosModularVehicle/Source/ChaosModularVehicleEngine/Private/ChaosModularVehicle/ModularVehicleSimulationCU.cpp
--- a/Engine/Plugins/Experimental/ChaosModularVehicle/Source/ChaosModularVehicleEngine/Private/ChaosModularVehicle/ModularVehicleSimulationCU.cpp
+++ b/Engine/Plugins/Experimental/ChaosModularVehicle/Source/ChaosModularVehicleEngine/Private/ChaosModularVehicle/ModularVehicleSimulationCU.cpp
@@ -151,7 +151,7 @@ void FModularVehicleSimulation::SimulateModuleTree(UWorld* InWorld, float DeltaS
 
 		if(SimModuleTree->GetSimTreeProcessingOrder() == ESimTreeProcessingOrder::ManualOverride)
 		{
-			VehicleSimulationCallback.Broadcast(DeltaSeconds, InputData, SimInputData, Proxy, RootParticle, SimModuleTree.Get());
+			VehicleSimulationCallback.Broadcast(this, InWorld, DeltaSeconds, InputData, SimInputData, Proxy, RootParticle, SimModuleTree.Get());
 		}
 
 		// Clear those Inputs that we don't want to remain set if the physics simulation thread ticks more frames than the GT
diff --git a/Engine/Plugins/Experimental/ChaosModularVehicle/Source/ChaosModularVehicleEngine/Public/ChaosModularVehicle/ModularVehicleSimulationCU.h b/Engine/Plugins/Experimental/ChaosModularVehicle/Source/ChaosModularVehicleEngine/Public/ChaosModularVehicle/ModularVehicleSimulationCU.h
--- a/Engine/Plugins/Experimental/ChaosModularVehicle/Source/ChaosModularVehicleEngine/Public/ChaosModularVehicle/ModularVehicleSimulationCU.h
+++ b/Engine/Plugins/Experimental/ChaosModularVehicle/Source/ChaosModularVehicleEngine/Public/ChaosModularVehicle/ModularVehicleSimulationCU.h
@@ -37,10 +37,10 @@ namespace Chaos
 {
 	class FClusterUnionPhysicsProxy;
 	class FSingleParticlePhysicsProxy;
-	DECLARE_MULTICAST_DELEGATE_SixParams(FModularVehicleSimulationCallback, float, const FModularVehicleAsyncInput&, FAllInputs&, IPhysicsProxyBase*, Chaos::FPBDRigidParticleHandle*, Chaos::FSimModuleTree*);
-	
 }
 
+class FModularVehicleSimulation;
+DECLARE_MULTICAST_DELEGATE_EightParams(FModularVehicleSimulationCallback, FModularVehicleSimulation*, UWorld*, float, const FModularVehicleAsyncInput&, Chaos::FAllInputs&, IPhysicsProxyBase*, Chaos::FPBDRigidParticleHandle*, Chaos::FSimModuleTree*);
 
 class FModularVehicleSimulation
 {
@@ -138,7 +138,7 @@ public:
 
 	Chaos::FPBDRigidParticleHandle* RootParticle = nullptr; // cached root particle
 	
-	Chaos::FModularVehicleSimulationCallback VehicleSimulationCallback;
+	FModularVehicleSimulationCallback VehicleSimulationCallback;
 };
 
 #undef UE_API
