Change code so not dependant on being a pawn in order to setup vehicle, plus some tightening of pointer use

#Synced-CL 48553644

diff --git a/Engine/Plugins/Experimental/ChaosModularVehicle/Source/ChaosModularVehicleEngine/Private/ChaosModularVehicle/ModularVehicleBaseComponent.cpp b/Engine/Plugins/Experimental/ChaosModularVehicle/Source/ChaosModularVehicleEngine/Private/ChaosModularVehicle/ModularVehicleBaseComponent.cpp
--- a/Engine/Plugins/Experimental/ChaosModularVehicle/Source/ChaosModularVehicleEngine/Private/ChaosModularVehicle/ModularVehicleBaseComponent.cpp
+++ b/Engine/Plugins/Experimental/ChaosModularVehicle/Source/ChaosModularVehicleEngine/Private/ChaosModularVehicle/ModularVehicleBaseComponent.cpp
@@ -1353,9 +1353,9 @@ void UModularVehicleBaseComponent::AssimilateComponentInputs(TArray<FModuleInput
 	OutCombinedInputs = InputConfig;
 
 	// append the input setup from all module sim components attached to same actor
-	if (APawn* Pawn = Cast<APawn>(GetOwner()))
+	if (AActor* Actor = GetOwner())
 	{
-		if (USceneComponent* RootComponent = Pawn->GetRootComponent())
+		if (USceneComponent* RootComponent = Actor->GetRootComponent())
 		{
 			TArray<USceneComponent*> ChildComponents;
 			RootComponent->GetChildrenComponents(true, ChildComponents);
@@ -1453,6 +1453,8 @@ void UModularVehicleBaseComponent::AddComponentToCluster(USceneComponent* Compon
 {
 	TArray<int32> BoneIds;
 
+	if (ClusterUnionComponent)
+	{
 	if (UGeometryCollectionComponent* GCComponent = Cast<UGeometryCollectionComponent>(Component))
 	{
 		BoneIds.Add(0);
@@ -1485,6 +1487,7 @@ void UModularVehicleBaseComponent::AddComponentToCluster(USceneComponent* Compon
 		}
 	}
 }
+}
 
 
 void UModularVehicleBaseComponent::AddOverlappingComponentsToCluster()
@@ -1517,14 +1520,14 @@ void UModularVehicleBaseComponent::AddOverlappingComponentsToCluster()
 
 bool UModularVehicleBaseComponent::AddComponentsFromOwnedActor()
 {
-	bool bSuccess = false;
+	bool bIsSuccess = false;
 
 	TArray<int32> BoneIds;
 	BoneIds.Add(0);
 
-	if (APawn* Pawn = Cast<APawn>(GetOwner()))
+	if (AActor* Actor = GetOwner())
 	{
-		if (UClusterUnionComponent* ClusterUnion = Pawn->GetComponentByClass<UClusterUnionComponent>())
+		if (UClusterUnionComponent* ClusterUnion = Actor->GetComponentByClass<UClusterUnionComponent>())
 		{
 			TArray<USceneComponent*> ChildComponents;
 			ClusterUnion->GetChildrenComponents(true, ChildComponents);
@@ -1532,41 +1535,43 @@ bool UModularVehicleBaseComponent::AddComponentsFromOwnedActor()
 			for (USceneComponent* Component : ChildComponents)
 			{
 				AddComponentToCluster(Component);
+				bIsSuccess = true;
 			}
 		}
 	}
 
-	return bSuccess;
+	return bIsSuccess;
 }
 
 void UModularVehicleBaseComponent::AddSimulationComponentsFromOwnedActor()
 {
 	// get raw simulation components that are not connected with any component - required for scene graph
-	if (APawn* Pawn = Cast<APawn>(GetOwner()))
+	if (AActor* Actor = GetOwner())
 	{
-		USceneComponent* RootComponent = Pawn->GetRootComponent();
-		TArray<USceneComponent*> ChildComponents;
-		RootComponent->GetChildrenComponents(false, ChildComponents);
+		if (USceneComponent* RootComponent = Actor->GetRootComponent())
+		{
+			TArray<USceneComponent*> ChildComponents;
+			RootComponent->GetChildrenComponents(false, ChildComponents);
 
-		int32 ComponentAddOrder = INDEX_NONE;
-		int ParentIndex=-1;
-		int TransformIndex=0;
-		Chaos::FSimTreeUpdates LatestTreeUpdates;
+			int32 ComponentAddOrder = INDEX_NONE;
+			int ParentIndex=-1;
+			int TransformIndex=0;
+			Chaos::FSimTreeUpdates LatestTreeUpdates;
 
-		for (USceneComponent* Component : ChildComponents)
-		{
-			if (UVehicleSimBaseComponent* SimulationComponent = Cast<UVehicleSimBaseComponent>(Component))
+			for (USceneComponent* Component : ChildComponents)
 			{
-				CreateIndependentSimComponents(RootComponent, SimulationComponent, ParentIndex, TransformIndex, LatestTreeUpdates);
-				NextTransformIndex++;
+				if (UVehicleSimBaseComponent* SimulationComponent = Cast<UVehicleSimBaseComponent>(Component))
+				{
+					CreateIndependentSimComponents(RootComponent, SimulationComponent, ParentIndex, TransformIndex, LatestTreeUpdates);
+					NextTransformIndex++;
+				}
 			}
-		}
-
-		Chaos::FSimTreeUpdates* NextTreeUpdates = nullptr;
-		NextTreeUpdates = &LatestTreeUpdates;
-		LastComponentAddIndex = ComponentAddOrder;
-		ActionTreeUpdates(NextTreeUpdates);
 
+			Chaos::FSimTreeUpdates* NextTreeUpdates = nullptr;
+			NextTreeUpdates = &LatestTreeUpdates;
+			LastComponentAddIndex = ComponentAddOrder;
+			ActionTreeUpdates(NextTreeUpdates);
+		}
 	}
 }
 
