Modular vehicle - Fix bug where removing and adding same component back will cause the simulation components to not be setup correctly on the client

#Synced-CL 48801579

diff --git a/Engine/Plugins/Experimental/ChaosModularVehicle/Source/ChaosModularVehicleEngine/Private/ChaosModularVehicle/ModularVehicleBaseComponent.cpp b/Engine/Plugins/Experimental/ChaosModularVehicle/Source/ChaosModularVehicleEngine/Private/ChaosModularVehicle/ModularVehicleBaseComponent.cpp
--- a/Engine/Plugins/Experimental/ChaosModularVehicle/Source/ChaosModularVehicleEngine/Private/ChaosModularVehicle/ModularVehicleBaseComponent.cpp
+++ b/Engine/Plugins/Experimental/ChaosModularVehicle/Source/ChaosModularVehicleEngine/Private/ChaosModularVehicle/ModularVehicleBaseComponent.cpp
@@ -365,7 +365,7 @@ void UModularVehicleBaseComponent::CreateAssociatedSimComponents(USceneComponent
 		// transforms in components, not just skeletal mesh bones.
 		if (bIsAnimationEnabled)
 		{
-			// if bone already exists then use that (seperate wheel and suspension modules can share same bone)
+			// if bone already exists then use that (separate wheel and suspension modules can share same bone)
 
 			int FoundIndex = INDEX_NONE;
 			if (ComponentBoneName != NAME_None)
@@ -484,7 +484,7 @@ void UModularVehicleBaseComponent::CreateIndependentSimComponents(USceneComponen
 		// transforms in components, not just skeletal mesh bones.
 		if (bIsAnimationEnabled)
 		{
-			// if bone already exists then use that (seperate wheel and suspension modules can share same bone)
+			// if bone already exists then use that (separate wheel and suspension modules can share same bone)
 
 			int FoundIndex = INDEX_NONE;
 			for (int I = 0; I < ModuleAnimationSetups.Num(); I++)
@@ -694,18 +694,47 @@ void UModularVehicleBaseComponent::TickComponent(float DeltaTime, ELevelTick Tic
 
 }
 
+void UModularVehicleBaseComponent::CreateOrUpdateComponentAddOrder(USceneComponent* InComponent, int32 InNextConstructionIndex)
+{
+	if (InComponent)
+	{
+		for (FConstructionData& Data : ConstructionDatas)
+		{
+			if (Data.Component == InComponent)
+			{
+				Data.ConstructionIndex = InNextConstructionIndex;
+				return;
+			}
+		}
+
+		FConstructionData ConstructionData;
+		ConstructionData.Component = InComponent;
+		ConstructionData.ConstructionIndex = InNextConstructionIndex;
+		ConstructionDatas.Add(ConstructionData);
+	}
+}
 
 int32 UModularVehicleBaseComponent::FindComponentAddOrder(USceneComponent* InComponent)
 {
 	int32 ComponentAddOrder = INDEX_NONE;
-	for (const FConstructionData& Data : ConstructionDatas)
+
+	if (InComponent)
 	{
-		if (Data.Component == InComponent)
+		for (const FConstructionData& Data : ConstructionDatas)
 		{
-			ComponentAddOrder = Data.ConstructionIndex;
-			break;
+			if (Data.Component == InComponent)
+			{
+				ComponentAddOrder = Data.ConstructionIndex;
+				break;
+			}
+		}
+
+		if (ComponentAddOrder == INDEX_NONE)
+		{
+			UE_LOG(LogModularBase, Warning, TEXT("FindComponentAddOrder failed to find component %s"), *InComponent->GetName());
 		}
 	}
+
 	return ComponentAddOrder;
 }
 
@@ -714,14 +743,22 @@ bool UModularVehicleBaseComponent::FindAndRemoveNextPendingUpdate(int32 NextInde
 	bool bFound = false;
 	for (const TPair<TObjectKey<USceneComponent>, Chaos::FSimTreeUpdates>& Kvp : PendingTreeUpdates)
 	{
-		if (FindComponentAddOrder(Kvp.Key.ResolveObjectPtr()) == NextIndex)
+		if (USceneComponent* Component = Kvp.Key.ResolveObjectPtr())
 		{
-			*OutData = PendingTreeUpdates.FindAndRemoveChecked(Kvp.Key.ResolveObjectPtr());
-			bFound = true;
-			break;
+			if (FindComponentAddOrder(Component) == NextIndex)
+			{
+				*OutData = PendingTreeUpdates.FindAndRemoveChecked(Component);
+				bFound = true;
+				break;
+			}
 		}
 	}
 
+	if (!bFound)
+	{
+		UE_LOG(LogModularBase, Warning, TEXT("FindAndRemoveNextPendingUpdate failed to find data for index %d"), NextIndex);
+	}
+
 	return bFound;
 }
 
@@ -1229,10 +1266,7 @@ void UModularVehicleBaseComponent::AddComponentToSimulationImpl(UPrimitiveCompon
 		if (ClusterUnionComponent->IsAuthority())
 		{
 			// retain the order that the components were constructed
-			FConstructionData ConstructionData;
-			ConstructionData.Component = InComponent;
-			ConstructionData.ConstructionIndex = NextConstructionIndex;
-			ConstructionDatas.Add(ConstructionData);
+			CreateOrUpdateComponentAddOrder(InComponent, NextConstructionIndex);
 			ComponentAddOrder = NextConstructionIndex++;
 		}
 		else
@@ -1321,6 +1355,15 @@ void UModularVehicleBaseComponent::RemoveComponentFromSimulationImpl(UPrimitiveC
 			}
 		}
 
+		if (ClusterUnionComponent->IsAuthority())
+		{
+			int32 Index = ConstructionDatas.IndexOfByPredicate([InComponent](const FConstructionData& InData) { return InData.Component == InComponent; });
+			if (Index != INDEX_NONE)
+			{
+				ConstructionDatas.RemoveAtSwap(Index);
+			}
+		}
+
 		ActionTreeUpdates(&LatestTreeUpdates);
 
 		NextTransformIndex--;
diff --git a/Engine/Plugins/Experimental/ChaosModularVehicle/Source/ChaosModularVehicleEngine/Public/ChaosModularVehicle/ModularVehicleBaseComponent.h b/Engine/Plugins/Experimental/ChaosModularVehicle/Source/ChaosModularVehicleEngine/Public/ChaosModularVehicle/ModularVehicleBaseComponent.h
--- a/Engine/Plugins/Experimental/ChaosModularVehicle/Source/ChaosModularVehicleEngine/Public/ChaosModularVehicle/ModularVehicleBaseComponent.h
+++ b/Engine/Plugins/Experimental/ChaosModularVehicle/Source/ChaosModularVehicleEngine/Public/ChaosModularVehicle/ModularVehicleBaseComponent.h
@@ -55,7 +55,7 @@ struct FConstructionData
 	GENERATED_USTRUCT_BODY()
 
 	UPROPERTY()
-	TObjectPtr<UPrimitiveComponent> Component = nullptr;
+	TObjectPtr<USceneComponent> Component = nullptr;
 
 	UPROPERTY()
 	int32 ConstructionIndex = INDEX_NONE;
@@ -366,6 +366,7 @@ protected:
 	UE_API int GenerateNewGuid();
 
 	UE_API int32 FindComponentAddOrder(USceneComponent* InComponent);
+	UE_API void CreateOrUpdateComponentAddOrder(USceneComponent* InComponent, int32 InComponentAddOrder);
 	UE_API bool FindAndRemoveNextPendingUpdate(int32 NextIndex, Chaos::FSimTreeUpdates* OutData);
 
 	UE_API void BroadcastModuleAddedEvent(const FName& SimType, int Guid, int TreeIndex);
