Null check the return of UModularVehicleBaseComponent::CreateNewCoreModule in UModularVehicleBaseComponent

#Synced-CL 50366066

diff --git a/Engine/Plugins/Experimental/ChaosModularVehicle/Source/ChaosModularVehicleEngine/Private/ChaosModularVehicle/ModularVehicleBaseComponent.cpp b/Engine/Plugins/Experimental/ChaosModularVehicle/Source/ChaosModularVehicleEngine/Private/ChaosModularVehicle/ModularVehicleBaseComponent.cpp
--- a/Engine/Plugins/Experimental/ChaosModularVehicle/Source/ChaosModularVehicleEngine/Private/ChaosModularVehicle/ModularVehicleBaseComponent.cpp
+++ b/Engine/Plugins/Experimental/ChaosModularVehicle/Source/ChaosModularVehicleEngine/Private/ChaosModularVehicle/ModularVehicleBaseComponent.cpp
@@ -317,90 +317,89 @@ void UModularVehicleBaseComponent::CreateAssociatedSimComponents(USceneComponent
 
 		ensure(TransformIndex < ChildParticles.Num());
 
-
-		FTransform ClusterUnionComponentTransform = ClusterUnionComponent->GetComponentTransform();
-		FTransform ComponentTransform = AttachedComponent->GetComponentTransform().GetRelativeTransform(ClusterUnionComponentTransform);
-
-		int TreeIndex = INDEX_NONE;
-
-		Chaos::ISimulationModuleBase* NewModule = ComponentInterface->CreateNewCoreModule();
-
-		TUniquePtr<Chaos::FSimModuleTree>& SimModuleTree = VehicleSimulationPT->AccessSimComponentTree();
-
-#if !(UE_BUILD_SHIPPING || UE_BUILD_TEST)
-		FString DebugString;
-		NewModule->GetDebugString(DebugString);
-		UE_LOG(LogModularBase, Log, TEXT("CreateAssociatedSimComponents Found Sim Component %s (%s)"), *NewModule->GetDebugName(), *DebugString);
-#endif
-
-		FVector LocalOffset(0, 0, 0);
-		FVector LocalOffsetCOM(0, 0, 0);
-		FTransform PhysicsTransform = FTransform::Identity;
-
-		FTransform InitialTransform = ChildParticles[TransformIndex].ChildToParent;
-		InitialTransform.SetLocation(InitialTransform.GetLocation());
-		NewModule->SetSimModuleTree(SimModuleTree.Get());
-		TreeIndex = TreeUpdatesOut.AddNodeBelow(ParentIndex, NewModule);
-		NewModule->SetTransformIndex(TransformIndex);
-		NewModule->SetParticleIndex(ChildParticles[TransformIndex].ParticleIdx);
-
-		NewModule->SetIntactTransform(FTransform::Identity);
-		NewModule->SetClusteredTransform(InitialTransform);
-		NewModule->SetClustered(true);
-
-		FVehicleComponentData& ComponentData = ComponentToPhysicsObjects.FindOrAdd(AttachedComponent);
-		ComponentData.Guid = GenerateNewGuid();
-		NewModule->SetGuid(ComponentData.Guid);
-
-		PhysicsGuidToComponent.Add(ComponentData.Guid, AttachedComponent);
-
-		NewModule->SetInitialParticleTransform(InitialTransform);
-		NewModule->SetComponentTransform(ComponentTransform);
-
-		FTransform ClusterredTransform(FQuat::Identity, InitialTransform.GetLocation());
-		NewModule->SetClusteredTransform(ClusterredTransform);
-
-		const bool bIsAnimationEnabled = ComponentInterface->GetAnimationEnabled();
-		const FName ComponentBoneName = ComponentInterface->GetBoneName();
-		// Now always using ModuleAnimationSetups even when there is no bone specified. ModuleAnimationSetups can be used to drive
-		// transforms in components, not just skeletal mesh bones.
-		if (bIsAnimationEnabled)
+		if (Chaos::ISimulationModuleBase* NewModule = ComponentInterface->CreateNewCoreModule())
 		{
-			// if bone already exists then use that (separate wheel and suspension modules can share same bone)
-
-			int FoundIndex = INDEX_NONE;
-			if (ComponentBoneName != NAME_None)
+			FTransform ClusterUnionComponentTransform = ClusterUnionComponent->GetComponentTransform();
+			FTransform ComponentTransform = AttachedComponent->GetComponentTransform().GetRelativeTransform(ClusterUnionComponentTransform);
+			int TreeIndex = INDEX_NONE;
+			
+			TUniquePtr<Chaos::FSimModuleTree>& SimModuleTree = VehicleSimulationPT->AccessSimComponentTree();
+
+	#if !(UE_BUILD_SHIPPING || UE_BUILD_TEST)
+			FString DebugString;
+			NewModule->GetDebugString(DebugString);
+			UE_LOG(LogModularBase, Log, TEXT("CreateAssociatedSimComponents Found Sim Component %s (%s)"), *NewModule->GetDebugName(), *DebugString);
+	#endif
+
+			FVector LocalOffset(0, 0, 0);
+			FVector LocalOffsetCOM(0, 0, 0);
+			FTransform PhysicsTransform = FTransform::Identity;
+
+			FTransform InitialTransform = ChildParticles[TransformIndex].ChildToParent;
+			InitialTransform.SetLocation(InitialTransform.GetLocation());
+			NewModule->SetSimModuleTree(SimModuleTree.Get());
+			TreeIndex = TreeUpdatesOut.AddNodeBelow(ParentIndex, NewModule);
+			NewModule->SetTransformIndex(TransformIndex);
+			NewModule->SetParticleIndex(ChildParticles[TransformIndex].ParticleIdx);
+
+			NewModule->SetIntactTransform(FTransform::Identity);
+			NewModule->SetClusteredTransform(InitialTransform);
+			NewModule->SetClustered(true);
+
+			FVehicleComponentData& ComponentData = ComponentToPhysicsObjects.FindOrAdd(AttachedComponent);
+			ComponentData.Guid = GenerateNewGuid();
+			NewModule->SetGuid(ComponentData.Guid);
+
+			PhysicsGuidToComponent.Add(ComponentData.Guid, AttachedComponent);
+
+			NewModule->SetInitialParticleTransform(InitialTransform);
+			NewModule->SetComponentTransform(ComponentTransform);
+
+			FTransform ClusterredTransform(FQuat::Identity, InitialTransform.GetLocation());
+			NewModule->SetClusteredTransform(ClusterredTransform);
+
+			const bool bIsAnimationEnabled = ComponentInterface->GetAnimationEnabled();
+			const FName ComponentBoneName = ComponentInterface->GetBoneName();
+			// Now always using ModuleAnimationSetups even when there is no bone specified. ModuleAnimationSetups can be used to drive
+			// transforms in components, not just skeletal mesh bones.
+			if (bIsAnimationEnabled)
 			{
-				for (int I = 0; I < ModuleAnimationSetups.Num(); I++)
+				// if bone already exists then use that (separate wheel and suspension modules can share same bone)
+
+				int FoundIndex = INDEX_NONE;
+				if (ComponentBoneName != NAME_None)
 				{
-					if (ModuleAnimationSetups[I].BoneName == ComponentBoneName)
+					for (int I = 0; I < ModuleAnimationSetups.Num(); I++)
 					{
-						FoundIndex = I;
-						break;
+						if (ModuleAnimationSetups[I].BoneName == ComponentBoneName)
+						{
+							FoundIndex = I;
+							break;
+						}
 					}
 				}
-			}
 
-			const FVector& ComponentAnimationOffset = ComponentInterface->GetAnimationOffset();
-			if (FoundIndex != INDEX_NONE)
-			{
-				NewModule->SetAnimationData(ComponentBoneName, ComponentAnimationOffset, FoundIndex);
-			}
-			else
-			{
-				NewModule->SetAnimationData(ComponentBoneName, ComponentAnimationOffset, ModuleAnimationSetups.Num());
-				FModuleAnimationSetup AnimSetup(NewModule->GetBoneName(), TransformIndex);
-				ModuleAnimationSetups.Add(AnimSetup);
+				const FVector& ComponentAnimationOffset = ComponentInterface->GetAnimationOffset();
+				if (FoundIndex != INDEX_NONE)
+				{
+					NewModule->SetAnimationData(ComponentBoneName, ComponentAnimationOffset, FoundIndex);
+				}
+				else
+				{
+					NewModule->SetAnimationData(ComponentBoneName, ComponentAnimationOffset, ModuleAnimationSetups.Num());
+					FModuleAnimationSetup AnimSetup(NewModule->GetBoneName(), TransformIndex);
+					ModuleAnimationSetups.Add(AnimSetup);
+				}
 			}
-		}
 
-		// store the tree index in the original sim component
-		ComponentInterface->SetTreeIndex(TreeIndex);
-		ParentIndex = TreeIndex;
+			// store the tree index in the original sim component
+			ComponentInterface->SetTreeIndex(TreeIndex);
+			ParentIndex = TreeIndex;
 
-		if (GetPhysicsProxy())
-		{
-			NewModule->OnConstruction_External(RootPhysicsObject);
+			if (GetPhysicsProxy())
+			{
+				NewModule->OnConstruction_External(RootPhysicsObject);
+			}
 		}
 	}
 
@@ -437,86 +436,86 @@ void UModularVehicleBaseComponent::CreateIndependentSimComponents(USceneComponen
 
 	if (IVehicleSimBaseComponentInterface* ComponentInterface = Cast<IVehicleSimBaseComponentInterface>(AttachedComponent))
 	{
-		FTransform RootComponentTransform = RootComponent->GetComponentTransform();
-		FTransform ComponentTransform = AttachedComponent->GetComponentTransform().GetRelativeTransform(RootComponentTransform);
-
-		int TreeIndex = INDEX_NONE;
-
-		Chaos::ISimulationModuleBase* NewModule = ComponentInterface->CreateNewCoreModule();
-
-		TUniquePtr<Chaos::FSimModuleTree>& SimModuleTree = VehicleSimulationPT->AccessSimComponentTree();
-
-#if !(UE_BUILD_SHIPPING || UE_BUILD_TEST)
-		FString DebugString;
-		NewModule->GetDebugString(DebugString);
-		UE_LOG(LogModularBase, Log, TEXT("CreateAssociatedSimComponents Found Sim Component %s (%s)"), *NewModule->GetDebugName(), *DebugString);
-#endif
-
-		FVector LocalOffset(0, 0, 0);
-		FVector LocalOffsetCOM(0, 0, 0);
-		FTransform PhysicsTransform = FTransform::Identity;
-
-		FTransform InitialTransform = ComponentTransform;
-		InitialTransform.SetLocation(InitialTransform.GetLocation());
-		NewModule->SetSimModuleTree(SimModuleTree.Get());
-		TreeIndex = TreeUpdatesOut.AddNodeBelow(ParentIndex, NewModule);
-		NewModule->SetTransformIndex(TransformIndex);
-		NewModule->SetParticleIndex(Chaos::FUniqueIdx());
-
-		NewModule->SetIntactTransform(FTransform::Identity);
-		NewModule->SetClusteredTransform(InitialTransform);
-		NewModule->SetClustered(true);
-
-		FVehicleComponentData& ComponentData = ComponentToPhysicsObjects.FindOrAdd(AttachedComponent);
-		ComponentData.Guid = GenerateNewGuid();
-		NewModule->SetGuid(ComponentData.Guid);
-
-		PhysicsGuidToComponent.Add(ComponentData.Guid, AttachedComponent);
-
-		NewModule->SetInitialParticleTransform(InitialTransform);
-		NewModule->SetComponentTransform(ComponentTransform);
-
-		FTransform ClusterredTransform(FQuat::Identity, InitialTransform.GetLocation());
-		NewModule->SetClusteredTransform(ClusterredTransform);
-
-		const bool bIsAnimationEnabled = ComponentInterface->GetAnimationEnabled();
-		const FName ComponentBoneName = ComponentInterface->GetBoneName();
-		// Changed behaviour - now always using ModuleAnimationSetups even when there is no bone setup. ModuleAnimationSetups is used to drive
-		// transforms in components, not just skeletal mesh bones.
-		if (bIsAnimationEnabled)
+		if (Chaos::ISimulationModuleBase* NewModule = ComponentInterface->CreateNewCoreModule())
 		{
-			// if bone already exists then use that (separate wheel and suspension modules can share same bone)
-
-			int FoundIndex = INDEX_NONE;
-			for (int I = 0; I < ModuleAnimationSetups.Num(); I++)
+			FTransform RootComponentTransform = RootComponent->GetComponentTransform();
+			FTransform ComponentTransform = AttachedComponent->GetComponentTransform().GetRelativeTransform(RootComponentTransform);
+
+			int TreeIndex = INDEX_NONE;
+			TUniquePtr<Chaos::FSimModuleTree>& SimModuleTree = VehicleSimulationPT->AccessSimComponentTree();
+
+	#if !(UE_BUILD_SHIPPING || UE_BUILD_TEST)
+			FString DebugString;
+			NewModule->GetDebugString(DebugString);
+			UE_LOG(LogModularBase, Log, TEXT("CreateAssociatedSimComponents Found Sim Component %s (%s)"), *NewModule->GetDebugName(), *DebugString);
+	#endif
+
+			FVector LocalOffset(0, 0, 0);
+			FVector LocalOffsetCOM(0, 0, 0);
+			FTransform PhysicsTransform = FTransform::Identity;
+
+			FTransform InitialTransform = ComponentTransform;
+			InitialTransform.SetLocation(InitialTransform.GetLocation());
+			NewModule->SetSimModuleTree(SimModuleTree.Get());
+			TreeIndex = TreeUpdatesOut.AddNodeBelow(ParentIndex, NewModule);
+			NewModule->SetTransformIndex(TransformIndex);
+			NewModule->SetParticleIndex(Chaos::FUniqueIdx());
+
+			NewModule->SetIntactTransform(FTransform::Identity);
+			NewModule->SetClusteredTransform(InitialTransform);
+			NewModule->SetClustered(true);
+
+			FVehicleComponentData& ComponentData = ComponentToPhysicsObjects.FindOrAdd(AttachedComponent);
+			ComponentData.Guid = GenerateNewGuid();
+			NewModule->SetGuid(ComponentData.Guid);
+
+			PhysicsGuidToComponent.Add(ComponentData.Guid, AttachedComponent);
+
+			NewModule->SetInitialParticleTransform(InitialTransform);
+			NewModule->SetComponentTransform(ComponentTransform);
+
+			FTransform ClusterredTransform(FQuat::Identity, InitialTransform.GetLocation());
+			NewModule->SetClusteredTransform(ClusterredTransform);
+
+			const bool bIsAnimationEnabled = ComponentInterface->GetAnimationEnabled();
+			const FName ComponentBoneName = ComponentInterface->GetBoneName();
+			// Changed behaviour - now always using ModuleAnimationSetups even when there is no bone setup. ModuleAnimationSetups is used to drive
+			// transforms in components, not just skeletal mesh bones.
+			if (bIsAnimationEnabled)
 			{
-				if (ModuleAnimationSetups[I].BoneName == ComponentBoneName)
+				// if bone already exists then use that (separate wheel and suspension modules can share same bone)
+
+				int FoundIndex = INDEX_NONE;
+				for (int I = 0; I < ModuleAnimationSetups.Num(); I++)
 				{
-					FoundIndex = I;
-					break;
+					if (ModuleAnimationSetups[I].BoneName == ComponentBoneName)
+					{
+						FoundIndex = I;
+						break;
+					}
 				}
-			}
 
-			const FVector& ComponentAnimationOffset = ComponentInterface->GetAnimationOffset();
-			if (FoundIndex != INDEX_NONE)
-			{
-				NewModule->SetAnimationData(ComponentBoneName, ComponentAnimationOffset, FoundIndex);
-			}
-			else
-			{
-				NewModule->SetAnimationData(ComponentBoneName, ComponentAnimationOffset, ModuleAnimationSetups.Num());
-				FModuleAnimationSetup AnimSetup(NewModule->GetBoneName(), TransformIndex);
-				ModuleAnimationSetups.Add(AnimSetup);
+				const FVector& ComponentAnimationOffset = ComponentInterface->GetAnimationOffset();
+				if (FoundIndex != INDEX_NONE)
+				{
+					NewModule->SetAnimationData(ComponentBoneName, ComponentAnimationOffset, FoundIndex);
+				}
+				else
+				{
+					NewModule->SetAnimationData(ComponentBoneName, ComponentAnimationOffset, ModuleAnimationSetups.Num());
+					FModuleAnimationSetup AnimSetup(NewModule->GetBoneName(), TransformIndex);
+					ModuleAnimationSetups.Add(AnimSetup);
+				}
 			}
-		}
 
-		// store the tree index in the original sim component
-		ComponentInterface->SetTreeIndex(TreeIndex);
-		ParentIndex = TreeIndex;
+			// store the tree index in the original sim component
+			ComponentInterface->SetTreeIndex(TreeIndex);
+			ParentIndex = TreeIndex;
 
-		if (GetPhysicsProxy())
-		{
-			NewModule->OnConstruction_External(RootPhysicsObject);
+			if (GetPhysicsProxy())
+			{
+				NewModule->OnConstruction_External(RootPhysicsObject);
+			}
 		}
 	}
 
@@ -1127,7 +1126,7 @@ void UModularVehicleBaseComponent::Update(float DeltaTime)
 
 		// game state container
 		AsyncInput->PhysicsInputs.StateInputs.StateInputContainer = StateInputContainer;
-	}
+}
 
 
 }
