Fix for how and when physics async output gets consumed by vehicles on the game thread. Move where async input is put together to ProduceInput, or it won't be seen by the just dispatched async frame requests.
- pop available outputs up to results time only, process all events from those
- only store the 2 most recent async outputs that bracket the results time at which we inteprolate
- move output consumption to PostUpdate (EndPhysicsFrame) so we consume output we blocked on immediately, and interpolate on that same frame with an output that received input from the current frame.
This reduces latency to sub game frame values as long as we also dispatch one physics frame in "the future" (EFixedDTAccumlatedTimeMode::AheadWithPushforward)

#Synced-CL 50226173

diff --git a/Engine/Plugins/Experimental/ChaosModularVehicle/Source/ChaosModularVehicleEngine/Private/ChaosModularVehicle/ChaosSimModuleManager.cpp b/Engine/Plugins/Experimental/ChaosModularVehicle/Source/ChaosModularVehicleEngine/Private/ChaosModularVehicle/ChaosSimModuleManager.cpp
--- a/Engine/Plugins/Experimental/ChaosModularVehicle/Source/ChaosModularVehicleEngine/Private/ChaosModularVehicle/ChaosSimModuleManager.cpp
+++ b/Engine/Plugins/Experimental/ChaosModularVehicle/Source/ChaosModularVehicleEngine/Private/ChaosModularVehicle/ChaosSimModuleManager.cpp
@@ -212,8 +212,6 @@ void FChaosSimModuleManager::Update(FPhysScene* PhysScene, float DeltaTime)
 
 	ScenePreTick(PhysScene, DeltaTime);
 
-	ParallelUpdateVehicles(DeltaTime);
-
 	if (World)
 	{
 		FChaosSimModuleManagerAsyncInput* AsyncInput = AsyncCallback->GetProducerInputData_External();
@@ -233,6 +231,8 @@ void FChaosSimModuleManager::Update(FPhysScene* PhysScene, float DeltaTime)
 
 void FChaosSimModuleManager::PostUpdate(FChaosScene* PhysScene)
 {
+	ParallelUpdateVehicles();
+
 	for (TWeakObjectPtr<UModularVehicleBaseComponent>& Vehicle : CUVehicles)
 	{
 		TStrongObjectPtr<UModularVehicleBaseComponent> StrongPtr = Vehicle.Pin();
@@ -253,139 +253,194 @@ void FChaosSimModuleManager::InjectInputs_External(int32 PhysicsStep, int32 NumS
 	FChaosSimModuleManagerAsyncInput* AsyncInput = AsyncCallback->GetProducerInputData_External();
 	check(AsyncInput);
 
+	AsyncInput->Reset();	// only want latest frame's data
+	AsyncInput->VehicleInputs.Reserve(CUVehicles.Num());
+	AsyncInput->World = World;
+	AsyncInput->Timestamp = Timestamp;
+	++Timestamp;
+
 	for (TWeakObjectPtr<UModularVehicleBaseComponent>& Vehicle : CUVehicles)
 	{
 		TStrongObjectPtr<UModularVehicleBaseComponent> StrongPtr = Vehicle.Pin();
 		if (StrongPtr.IsValid())
 		{
-			StrongPtr->ProduceInput(PhysicsStep, NumSteps);
+			TUniquePtr<FModularVehicleAsyncInput> CurInput = MakeUnique<FModularVehicleAsyncInput>();
+			AsyncInput->VehicleInputs.Add(MoveTemp(CurInput));
+			StrongPtr->ProduceInput(PhysicsStep, NumSteps, AsyncInput->VehicleInputs.Last().Get());
 		}
 	}
+}
 
+void FSimModuleOutputRecord::ConsumeOutput(Chaos::TSimCallbackOutputHandle<FChaosSimModuleManagerAsyncOutput>&& Output)
+{
+	if (!CachedOutput_0)
+	{
+		CachedOutput_0 = MoveTemp(Output);
+	}
+	else if (!CachedOutput_1)
+	{
+		CachedOutput_1 = MoveTemp(Output);
+	}
+	else
+	{
+		(bPreviousOutputIs0 ? CachedOutput_0 : CachedOutput_1) = MoveTemp(Output);
+		bPreviousOutputIs0 = !bPreviousOutputIs0;
+	}
 }
 
+FChaosSimModuleManagerAsyncOutput* FSimModuleOutputRecord::GetPreviousOutput()
+{
+	if (CachedOutput_0 && bPreviousOutputIs0)
+	{
+		return CachedOutput_0.Get();
+	}
+	if (CachedOutput_1 && !bPreviousOutputIs0)
+	{
+		return CachedOutput_1.Get();
+	}
+	return nullptr;
+}
 
-void FChaosSimModuleManager::ParallelUpdateVehicles(float DeltaSeconds)
+FChaosSimModuleManagerAsyncOutput* FSimModuleOutputRecord::GetNextOutput()
 {
+	if (CachedOutput_1 && bPreviousOutputIs0)
+	{
+		return CachedOutput_1.Get();
+	}
+	if (CachedOutput_0 && !bPreviousOutputIs0)
+	{
+		return CachedOutput_0.Get();
+	}
+	return nullptr;
+}
 
-	FChaosSimModuleManagerAsyncInput* AsyncInput = AsyncCallback->GetProducerInputData_External();
+const FChaosSimModuleManagerAsyncOutput* FSimModuleOutputRecord::GetPreviousOutput() const
+{
+	if (CachedOutput_0 && bPreviousOutputIs0)
+	{
+		return CachedOutput_0.Get();
+	}
+	if (CachedOutput_1 && !bPreviousOutputIs0)
+	{
+		return CachedOutput_1.Get();
+	}
+	return nullptr;
+}
 
-	AsyncInput->Reset();	// only want latest frame's data
+const FChaosSimModuleManagerAsyncOutput* FSimModuleOutputRecord::GetNextOutput() const
+{
+	if (CachedOutput_1 && bPreviousOutputIs0)
+	{
+		return CachedOutput_1.Get();
+	}
+	if (CachedOutput_0 && !bPreviousOutputIs0)
+	{
+		return CachedOutput_0.Get();
+	}
+	return nullptr;
+}
 
+double FSimModuleOutputRecord::GetLatestOutputStartTime() const
+{
+	double LatestOutputStartTime = -DBL_MAX;
+	if (CachedOutput_0 && CachedOutput_1)
 	{
-		// We pass pointers from TArray so this reserve is critical. Otherwise realloc happens
-		AsyncInput->VehicleInputs.Reserve(CUVehicles.Num());
-		AsyncInput->Timestamp = Timestamp;
-		AsyncInput->World = Scene.GetOwningWorld();
+		LatestOutputStartTime = (bPreviousOutputIs0 ? CachedOutput_1 : CachedOutput_0)->InternalTime;
 	}
+	else if (CachedOutput_0)
+	{
+		LatestOutputStartTime = CachedOutput_0->InternalTime;
+	}
+	return LatestOutputStartTime;
+}
 
-	// Grab all outputs for processing, even future ones for interpolation.
+double FSimModuleOutputRecord::GetInterpolationFactor(double AtInternalTime) const
+{
+	const FChaosSimModuleManagerAsyncOutput* PreviousOutput = GetPreviousOutput();
+	const FChaosSimModuleManagerAsyncOutput* NextOutput = GetNextOutput();
+	double Alpha = 0.f;
+	if (NextOutput)
 	{
-		Chaos::TSimCallbackOutputHandle<FChaosSimModuleManagerAsyncOutput> AsyncOutputLatest;
-		while ((AsyncOutputLatest = AsyncCallback->PopFutureOutputData_External()))
+		const double Denom = NextOutput->InternalTime - PreviousOutput->InternalTime;
+		if (Denom > SMALL_NUMBER)
 		{
-			PendingOutputs.Emplace(MoveTemp(AsyncOutputLatest));
+			Alpha = FMath::Clamp((AtInternalTime - PreviousOutput->InternalTime) / Denom, 0.0f, 1.0f);
 		}
 	}
+	return Alpha;
+}
+
+void FSimModuleOutputRecord::Clear()
+{
+	CachedOutput_0 = Chaos::TSimCallbackOutputHandle<FChaosSimModuleManagerAsyncOutput>();
+	CachedOutput_1 = Chaos::TSimCallbackOutputHandle<FChaosSimModuleManagerAsyncOutput>();
+	bPreviousOutputIs0 = true;
+}
+
+void FChaosSimModuleManager::ParallelUpdateVehicles()
+{
+	// Friendly reminder: Results time is the time at the END of an async step
+	const double ResultsTime = AsyncCallback->GetSolver()->GetPhysicsResultsTime_External();
+	const double AsyncDt = AsyncCallback->GetSolver()->GetAsyncDeltaTime();
 
-	// ensure we don't miss any module creation events
 	TArray<Chaos::FCreatedModules> CombinedNewlyCreatedModuleGuids;
 	Chaos::FCreatedModules ModulesEventData;
 	CombinedNewlyCreatedModuleGuids.Init(ModulesEventData, CUVehicles.Num());
 
-	for (int32 OutIdx = 0; OutIdx < PendingOutputs.Num(); ++OutIdx)
+	// We need to pop new output data only if results time is beyond the latest cached output's step end time
+	// Otherwise we can still use the ones we have stored in OutputRecord to interpolate
+	// OutputRecord stores FChaosSimModuleManagerAsyncOutput which "InternalTime" si the time at the BEGINNING of the step that produced that output
+	const double LatestOutputResultsTime = OutputRecord.GetLatestOutputStartTime() + AsyncDt;
+	if (ResultsTime > LatestOutputResultsTime)
 	{
-		if (FChaosSimModuleManagerAsyncOutput* AsyncOutput = PendingOutputs[OutIdx].Get())
+		Chaos::TSimCallbackOutputHandle<FChaosSimModuleManagerAsyncOutput> NextAsyncOutput;
+		while ((NextAsyncOutput = AsyncCallback->PopOutputData_External()))
 		{
-			for (int32 VehicleIdx = 0; VehicleIdx < AsyncOutput->VehicleOutputs.Num(); ++VehicleIdx)
+			// We process all events we find to never miss one
+			for (int32 VehicleIdx = 0; VehicleIdx < NextAsyncOutput->VehicleOutputs.Num(); ++VehicleIdx)
 			{
-				TArray<Chaos::FCreatedModule>& NewlyCreatedModuleGuids = AsyncOutput->VehicleOutputs[VehicleIdx]->VehicleSimOutput.NewlyCreatedModuleGuids;
-
+				TArray<Chaos::FCreatedModule>& NewlyCreatedModuleGuids = NextAsyncOutput->VehicleOutputs[VehicleIdx]->VehicleSimOutput.NewlyCreatedModuleGuids;
 				if (NewlyCreatedModuleGuids.Num() > 0 && VehicleIdx < CombinedNewlyCreatedModuleGuids.Num())
 				{
 					CombinedNewlyCreatedModuleGuids[VehicleIdx].ModuleEvents.Append(NewlyCreatedModuleGuids);
 				}
 			}
+			OutputRecord.ConsumeOutput(MoveTemp(NextAsyncOutput));
 		}
 	}
 
-	// Since we are in pre-physics, delta seconds is not accounted for in external time yet
-	const double ResultsTime = AsyncCallback->GetSolver()->GetPhysicsResultsTime_External();
-
-	// Find index of first non-consumable output (first one after current time)
-	int32 LastOutputIdx = 0;
-	for (; LastOutputIdx < PendingOutputs.Num(); ++LastOutputIdx)
-	{
-		if (PendingOutputs[LastOutputIdx]->InternalTime > ResultsTime)
-		{
-			break;
-		}
-	}
-
-	// Cache the last consumed output for interpolation
-	if (LastOutputIdx > 0)
-	{
-		LatestOutput = MoveTemp(PendingOutputs[LastOutputIdx - 1]);
-	}
-
-	// Remove all consumed outputs
-	{
-		TArray<Chaos::TSimCallbackOutputHandle<FChaosSimModuleManagerAsyncOutput>> NewPendingOutputs;
-		for (int32 OutputIdx = LastOutputIdx; OutputIdx < PendingOutputs.Num(); ++OutputIdx)
-		{
-			NewPendingOutputs.Emplace(MoveTemp(PendingOutputs[OutputIdx]));
-		}
-		PendingOutputs = MoveTemp(NewPendingOutputs);
-	}
-
-	// It's possible we will end up multiple frames ahead of output, take the latest ready output.
-	Chaos::TSimCallbackOutputHandle<FChaosSimModuleManagerAsyncOutput> AsyncOutput;
-	Chaos::TSimCallbackOutputHandle<FChaosSimModuleManagerAsyncOutput> AsyncOutputLatest;
-	while ((AsyncOutputLatest = AsyncCallback->PopOutputData_External()))
-	{
-		AsyncOutput = MoveTemp(AsyncOutputLatest);
-	}
-
-	if (UWorld* World = Scene.GetOwningWorld())
+	FChaosSimModuleManagerAsyncOutput* PreviousOutput = OutputRecord.GetPreviousOutput();
+	FChaosSimModuleManagerAsyncOutput* NextOutput = OutputRecord.GetNextOutput();
+	UWorld* World = Scene.GetOwningWorld();
+	FChaosSimModuleManagerAsyncInput* AsyncInput = AsyncCallback->GetProducerInputData_External();
+	if (World && PreviousOutput && AsyncInput)
 	{
-		int32 NumVehiclesInActiveBatch = 0;
+		// InternalTime is the time at the beginning of the step, not at the end, whereas ResultsTime is the end time of the step (where we got to by taking the step)
+		// OutputRecord deals with start times so we need to subtract AsyncDt from ResultsTime in what we pass to GetInterpolationFactor
+		const double InterpolationBaseTime = ResultsTime - AsyncDt;
+		float Alpha = OutputRecord.GetInterpolationFactor(InterpolationBaseTime);
 
 		for (TWeakObjectPtr<UModularVehicleBaseComponent>& Vehicle : CUVehicles)
 		{
 			TStrongObjectPtr<UModularVehicleBaseComponent> StrongPtr = Vehicle.Pin();
 			if (StrongPtr.IsValid())
 			{
-				auto NextOutput = PendingOutputs.Num() > 0 ? PendingOutputs[0].Get() : nullptr;
-				float Alpha = 0.f;
-				if (NextOutput && LatestOutput)
-				{
-					const float Denom = NextOutput->InternalTime - LatestOutput->InternalTime;
-					if (Denom > SMALL_NUMBER)
-					{
-						Alpha = (ResultsTime - LatestOutput->InternalTime) / Denom;
-					}
-				}
-
-				AsyncInput->VehicleInputs.Add(StrongPtr->SetCurrentAsyncData(AsyncInput->VehicleInputs.Num(), LatestOutput.Get(), NextOutput, Alpha, Timestamp));
+				StrongPtr->SetCurrentAsyncData(AsyncInput->VehicleInputs.Num(), PreviousOutput, NextOutput, Alpha, Timestamp);
 			}
 		}
-
 	}
 
-	++Timestamp;
-
 	bool ForceSingleThread = !GSimModuleDebugParams.EnableMultithreading;
-
 	{
 		const TArray<TWeakObjectPtr<UModularVehicleBaseComponent>>& AwakeVehiclesBatch = CUVehicles;
 		const TArray<Chaos::FCreatedModules>& ModuleEvents = CombinedNewlyCreatedModuleGuids;
-		auto LambdaParallelUpdate2 = [DeltaSeconds, &AwakeVehiclesBatch, &ModuleEvents](int32 Idx)
+		auto LambdaParallelUpdate2 = [&AwakeVehiclesBatch, &ModuleEvents](int32 Idx)
 			{
 				TWeakObjectPtr<UModularVehicleBaseComponent> Vehicle = AwakeVehiclesBatch[Idx];
 				TStrongObjectPtr<UModularVehicleBaseComponent> StrongPtr = Vehicle.Pin();
 				if (StrongPtr.IsValid())
 				{
-					StrongPtr->ParallelUpdate(DeltaSeconds, ModuleEvents[Idx]); // gets output state from PT
+					StrongPtr->ParallelUpdate(ModuleEvents[Idx]); // gets output state from PT
 				}
 			};
 
diff --git a/Engine/Plugins/Experimental/ChaosModularVehicle/Source/ChaosModularVehicleEngine/Private/ChaosModularVehicle/ModularVehicleBaseComponent.cpp b/Engine/Plugins/Experimental/ChaosModularVehicle/Source/ChaosModularVehicleEngine/Private/ChaosModularVehicle/ModularVehicleBaseComponent.cpp
--- a/Engine/Plugins/Experimental/ChaosModularVehicle/Source/ChaosModularVehicleEngine/Private/ChaosModularVehicle/ModularVehicleBaseComponent.cpp
+++ b/Engine/Plugins/Experimental/ChaosModularVehicle/Source/ChaosModularVehicleEngine/Private/ChaosModularVehicle/ModularVehicleBaseComponent.cpp
@@ -120,12 +120,35 @@ bool UModularVehicleBaseComponent::IsLocallyControlled() const
 	return false;
 }
 
-void UModularVehicleBaseComponent::ProduceInput(int32 PhysicsStep, int32 NumSteps)
+void UModularVehicleBaseComponent::ProduceInput(int32 PhysicsStep, int32 NumSteps, FModularVehicleAsyncInput* AsyncInput)
 {
 	if (InputProducer)
 	{
 		InputProducer->ProduceInput(PhysicsStep, NumSteps, InputNameMap, InputsContainer);
 	}
+
+	if (AsyncInput && GetPhysicsProxy())
+	{
+		AsyncInput->SetVehicle(this);
+		CurAsyncInput = AsyncInput;
+		AsyncInput->Proxy = GetPhysicsProxy();
+		AsyncInput->PhysicsInputs.NetworkInputs.VehicleInputs.KeepAwake = bKeepVehicleAwake;
+
+		// All control inputs
+		AsyncInput->PhysicsInputs.NetworkInputs.VehicleInputs.Container = InputsContainer;
+
+		FCollisionQueryParams TraceParams(NAME_None, FCollisionQueryParams::GetUnknownStatId(), false, nullptr);
+		TraceParams.bReturnPhysicalMaterial = true;	// we need this to get the surface friction coefficient
+		TraceParams.AddIgnoredActors(ActorsToIgnore);
+		TraceParams.bTraceComplex = bSuspensionTraceComplex;
+		AsyncInput->PhysicsInputs.CollisionChannel = SuspensionCollisionChannel;
+		AsyncInput->PhysicsInputs.TraceParams = TraceParams;
+		AsyncInput->PhysicsInputs.TraceCollisionResponse = SuspensionTraceCollisionResponses;
+		AsyncInput->PhysicsInputs.TraceType = TraceType;
+
+		// game state container
+		AsyncInput->PhysicsInputs.StateInputs.StateInputContainer = StateInputContainer;
+	}
 }
 
 void UModularVehicleBaseComponent::OnCreatePhysicsState()
@@ -847,23 +870,18 @@ void UModularVehicleBaseComponent::ServerUpdateState_Implementation(const FModul
 }
 
 
-TUniquePtr<FModularVehicleAsyncInput> UModularVehicleBaseComponent::SetCurrentAsyncData(int32 InputIdx, FChaosSimModuleManagerAsyncOutput* CurOutput, FChaosSimModuleManagerAsyncOutput* NextOutput, float Alpha, int32 VehicleManagerTimestamp)
+void UModularVehicleBaseComponent::SetCurrentAsyncData(int32 InputIdx, FChaosSimModuleManagerAsyncOutput* CurOutput, FChaosSimModuleManagerAsyncOutput* NextOutput, float Alpha, int32 VehicleManagerTimestamp)
 {
-	TUniquePtr<FModularVehicleAsyncInput> CurInput = MakeUnique<FModularVehicleAsyncInput>();
-	SetCurrentAsyncDataInternal(CurInput.Get(), InputIdx, CurOutput, NextOutput, Alpha, VehicleManagerTimestamp);
-	return CurInput;
+	SetCurrentAsyncDataInternal(InputIdx, CurOutput, NextOutput, Alpha, VehicleManagerTimestamp);
 }
 
 /************************************************************************/
 /* Setup the current async I/O data                                     */
 /************************************************************************/
-void UModularVehicleBaseComponent::SetCurrentAsyncDataInternal(FModularVehicleAsyncInput* CurInput, int32 InputIdx, FChaosSimModuleManagerAsyncOutput* CurOutput, FChaosSimModuleManagerAsyncOutput* NextOutput, float Alpha, int32 VehicleManagerTimestamp)
+void UModularVehicleBaseComponent::SetCurrentAsyncDataInternal(int32 InputIdx, FChaosSimModuleManagerAsyncOutput* CurOutput, FChaosSimModuleManagerAsyncOutput* NextOutput, float Alpha, int32 VehicleManagerTimestamp)
 {
-	ensure(CurAsyncInput == nullptr);	//should be reset after it was filled
 	ensure(CurAsyncOutput == nullptr);	//should get reset after update is done
 
-	CurAsyncInput = CurInput;
-	CurAsyncInput->SetVehicle(this);
 	NextAsyncOutput = nullptr;
 	OutputInterpAlpha = 0.f;
 
@@ -977,7 +995,7 @@ void UModularVehicleBaseComponent::SetClusterComponent(UClusterUnionComponent* I
 /***************************************************************************/
 /* READ OUTPUT DATA - Access the async output data from the Physics Thread */
 /***************************************************************************/
-void UModularVehicleBaseComponent::ParallelUpdate(const float DeltaTime, const Chaos::FCreatedModules& ModuleEvents)
+void UModularVehicleBaseComponent::ParallelUpdate(const Chaos::FCreatedModules& ModuleEvents)
 {
 	if (FModularVehicleAsyncOutput* CurrentOutput = static_cast<FModularVehicleAsyncOutput*>(CurAsyncOutput))
 	{
@@ -1098,39 +1116,11 @@ void UModularVehicleBaseComponent::ParallelUpdate(const float DeltaTime, const C
 }
 
 
-/************************************************************************/
-/* PASS ANY INPUTS TO THE PHYSICS THREAD SIMULATION IN HERE             */
-/************************************************************************/
 void UModularVehicleBaseComponent::Update(float DeltaTime)
 {
-	if (CurAsyncInput && GetPhysicsProxy())
-	{
-
-		CurAsyncInput->Proxy = GetPhysicsProxy();
-
-		FModularVehicleAsyncInput* AsyncInput = static_cast<FModularVehicleAsyncInput*>(CurAsyncInput);
-
-		AsyncInput->PhysicsInputs.NetworkInputs.VehicleInputs.KeepAwake = bKeepVehicleAwake;
-
-		// All control inputs
-		AsyncInput->PhysicsInputs.NetworkInputs.VehicleInputs.Container = InputsContainer;
-
-		FCollisionQueryParams TraceParams(NAME_None, FCollisionQueryParams::GetUnknownStatId(), false, nullptr);
-		TraceParams.bReturnPhysicalMaterial = true;	// we need this to get the surface friction coefficient
-		TraceParams.AddIgnoredActors(ActorsToIgnore);
-		TraceParams.bTraceComplex = bSuspensionTraceComplex;
-		AsyncInput->PhysicsInputs.CollisionChannel = SuspensionCollisionChannel;
-		AsyncInput->PhysicsInputs.TraceParams = TraceParams;
-		AsyncInput->PhysicsInputs.TraceCollisionResponse = SuspensionTraceCollisionResponses;
-		AsyncInput->PhysicsInputs.TraceType = TraceType;
-
-		// game state container
-		AsyncInput->PhysicsInputs.StateInputs.StateInputContainer = StateInputContainer;
 }
 
 
-}
-
 void UModularVehicleBaseComponent::PostUpdate()
 {
 	ensure(PVehicleOutput);
diff --git a/Engine/Plugins/Experimental/ChaosModularVehicle/Source/ChaosModularVehicleEngine/Public/ChaosModularVehicle/ChaosSimModuleManager.h b/Engine/Plugins/Experimental/ChaosModularVehicle/Source/ChaosModularVehicleEngine/Public/ChaosModularVehicle/ChaosSimModuleManager.h
--- a/Engine/Plugins/Experimental/ChaosModularVehicle/Source/ChaosModularVehicleEngine/Public/ChaosModularVehicle/ChaosSimModuleManager.h
+++ b/Engine/Plugins/Experimental/ChaosModularVehicle/Source/ChaosModularVehicleEngine/Public/ChaosModularVehicle/ChaosSimModuleManager.h
@@ -18,6 +18,38 @@ class AHUD;
 class FPhysScene_Chaos;
 struct FChaosSimModuleManagerAsyncOutput;
 
+// Helper class to store physics async output results
+// Todo: generalize this as a Chaos helper template class
+class FSimModuleOutputRecord
+{
+public:
+	// Consumes an async output. Assumes the output is dated later than previously consumed outputs.
+	// Uses move semantics, and will replace the oldest stored output.
+	void ConsumeOutput(Chaos::TSimCallbackOutputHandle<FChaosSimModuleManagerAsyncOutput>&& Output);
+
+	// Returns the oldest stored output
+	const FChaosSimModuleManagerAsyncOutput* GetPreviousOutput() const;
+	FChaosSimModuleManagerAsyncOutput* GetPreviousOutput();
+
+	// Returns the most recent stored output
+	const FChaosSimModuleManagerAsyncOutput* GetNextOutput() const;
+	FChaosSimModuleManagerAsyncOutput* GetNextOutput();
+
+	// Returns the start time (InternalTime) of the most recent output
+	double GetLatestOutputStartTime() const;
+
+	// Returns the interpolation factor
+	double GetInterpolationFactor(double AtInternalTime) const;
+
+	// Clears (and deletes) stored outputs
+	void Clear();
+
+private:
+	Chaos::TSimCallbackOutputHandle<FChaosSimModuleManagerAsyncOutput> CachedOutput_0;
+	Chaos::TSimCallbackOutputHandle<FChaosSimModuleManagerAsyncOutput> CachedOutput_1;
+	bool bPreviousOutputIs0 = true;
+};
+
 class FChaosSimModuleManager
 {
 public:
@@ -65,7 +97,7 @@ public:
 	/** Called on GT but before Physics thread runs - PT tick rate */
 	UE_API void InjectInputs_External(int32 PhysicsStep, int32 NumSteps);
 
-	UE_API void ParallelUpdateVehicles(float DeltaSeconds);
+	UE_API void ParallelUpdateVehicles();
 
 	/** Find a vehicle manager from an FPhysScene */
 	static UE_API FChaosSimModuleManager* GetManagerFromScene(FPhysScene* PhysScene);
@@ -99,10 +131,8 @@ private:
 	int32 Timestamp;
 	int32 SubStepCount;
 
-	// async input/output state
-	TArray<Chaos::TSimCallbackOutputHandle<FChaosSimModuleManagerAsyncOutput>> PendingOutputs;
-	Chaos::TSimCallbackOutputHandle<FChaosSimModuleManagerAsyncOutput> LatestOutput;
-
+	// Record used to store the two most recent outputs which time surrounds the time we interpolate at
+	FSimModuleOutputRecord OutputRecord;
 };
 
 #undef UE_API
diff --git a/Engine/Plugins/Experimental/ChaosModularVehicle/Source/ChaosModularVehicleEngine/Public/ChaosModularVehicle/ModularVehicleBaseComponent.h b/Engine/Plugins/Experimental/ChaosModularVehicle/Source/ChaosModularVehicleEngine/Public/ChaosModularVehicle/ModularVehicleBaseComponent.h
--- a/Engine/Plugins/Experimental/ChaosModularVehicle/Source/ChaosModularVehicleEngine/Public/ChaosModularVehicle/ModularVehicleBaseComponent.h
+++ b/Engine/Plugins/Experimental/ChaosModularVehicle/Source/ChaosModularVehicleEngine/Public/ChaosModularVehicle/ModularVehicleBaseComponent.h
@@ -142,7 +142,7 @@ public:
 	UE_API virtual void BeginPlay() override;
 	UE_API virtual void TickComponent(float DeltaTime, ELevelTick TickType, FActorComponentTickFunction* ThisTickFunction) override;
 
-	UE_API void ProduceInput(int32 PhysicsStep, int32 NumSteps);
+	UE_API void ProduceInput(int32 PhysicsStep, int32 NumSteps, FModularVehicleAsyncInput* AsyncInput);
 
 	// Create sim components hierarchy from parent component - used by original cluster code union path
 	UE_API virtual void CreateAssociatedSimComponents(USceneComponent* ParentComponent, USceneComponent* AttachedComponent, int ParentIndex, int TransformIndex, Chaos::FSimTreeUpdates& TreeUpdatesOut);
@@ -152,9 +152,9 @@ public:
 
 	UE_API void PreTickGT(float DeltaTime);
 	UE_API void UpdateState(float DeltaTime);
-	UE_API TUniquePtr<FModularVehicleAsyncInput> SetCurrentAsyncData(int32 InputIdx, FChaosSimModuleManagerAsyncOutput* CurOutput, FChaosSimModuleManagerAsyncOutput* NextOutput, float Alpha, int32 VehicleManagerTimestamp);
+	UE_API void SetCurrentAsyncData(int32 InputIdx, FChaosSimModuleManagerAsyncOutput* CurOutput, FChaosSimModuleManagerAsyncOutput* NextOutput, float Alpha, int32 VehicleManagerTimestamp);
 
-	UE_API void ParallelUpdate(const float DeltaTime, const Chaos::FCreatedModules& ModuleEvents);
+	UE_API void ParallelUpdate(const Chaos::FCreatedModules& ModuleEvents);
 	UE_API void Update(float DeltaTime);
 	UE_API void PostUpdate();
 	UE_API void FinalizeSimCallbackData(FChaosSimModuleManagerAsyncInput& Input);
@@ -359,7 +359,7 @@ protected:
 
 	UE_API void ActionTreeUpdates(Chaos::FSimTreeUpdates* NextTreeUpdates);
 
-	UE_API void SetCurrentAsyncDataInternal(FModularVehicleAsyncInput* CurInput, int32 InputIdx, FChaosSimModuleManagerAsyncOutput* CurOutput, FChaosSimModuleManagerAsyncOutput* NextOutput, float Alpha, int32 VehicleManagerTimestamp);
+	UE_API void SetCurrentAsyncDataInternal(int32 InputIdx, FChaosSimModuleManagerAsyncOutput* CurOutput, FChaosSimModuleManagerAsyncOutput* NextOutput, float Alpha, int32 VehicleManagerTimestamp);
 	UE_API int32 FindParentsLastSimComponent(const USceneComponent* AttachedComponent);
 
 	UE_API IPhysicsProxyBase* GetPhysicsProxy() const;
