Modular Vehicle: Fix performance issue with FillOutputState

#Synced-CL 50428374

diff --git a/Engine/Plugins/Experimental/ChaosModularVehicle/Source/ChaosModularVehicleEngine/Private/ChaosModularVehicle/ModularVehicleSimulationCU.cpp b/Engine/Plugins/Experimental/ChaosModularVehicle/Source/ChaosModularVehicleEngine/Private/ChaosModularVehicle/ModularVehicleSimulationCU.cpp
--- a/Engine/Plugins/Experimental/ChaosModularVehicle/Source/ChaosModularVehicleEngine/Private/ChaosModularVehicle/ModularVehicleSimulationCU.cpp
+++ b/Engine/Plugins/Experimental/ChaosModularVehicle/Source/ChaosModularVehicleEngine/Private/ChaosModularVehicle/ModularVehicleSimulationCU.cpp
@@ -469,6 +469,8 @@ void FModularVehicleSimulation::ApplyDeferredForces(IPhysicsProxyBase* Proxy)
 
 void FModularVehicleSimulation::FillOutputState(FModularVehicleAsyncOutput& Output)
 {
+	QUICK_SCOPE_CYCLE_COUNTER(Stat_FModularVehicleSimulation_FillOutputState);
+
 	Output.VehicleSimOutput.NewlyCreatedModuleGuids = NewlyCreatedModuleGuids;
 	NewlyCreatedModuleGuids.Empty();
 
diff --git a/Engine/Source/Runtime/Experimental/ChaosVehicles/ChaosVehiclesCore/Private/SimModule/SimulationModuleBase.cpp b/Engine/Source/Runtime/Experimental/ChaosVehicles/ChaosVehiclesCore/Private/SimModule/SimulationModuleBase.cpp
--- a/Engine/Source/Runtime/Experimental/ChaosVehicles/ChaosVehiclesCore/Private/SimModule/SimulationModuleBase.cpp
+++ b/Engine/Source/Runtime/Experimental/ChaosVehicles/ChaosVehiclesCore/Private/SimModule/SimulationModuleBase.cpp
@@ -149,7 +149,7 @@ void FSimOutputData::FillOutputState(const ISimulationModuleBase* SimModule)
 		ModuleGuid = SimModule->GetGuid();
 	}
 
-#if !(UE_BUILD_SHIPPING || UE_BUILD_TEST)
+#if UE_BUILD_DEBUG
 	if (SimModule)
 	{	
 		DebugString.Empty();
