Possible stale pointer if Wheel->SetGroundInteraction doesn't get called within PerformAdditionalSimWork

#Synced-CL 50695105

diff --git a/Engine/Plugins/Experimental/ChaosModularVehicle/Source/ChaosModularVehicleEngine/Private/ChaosModularVehicle/ModularVehicleSimulationCU.cpp b/Engine/Plugins/Experimental/ChaosModularVehicle/Source/ChaosModularVehicleEngine/Private/ChaosModularVehicle/ModularVehicleSimulationCU.cpp
--- a/Engine/Plugins/Experimental/ChaosModularVehicle/Source/ChaosModularVehicleEngine/Private/ChaosModularVehicle/ModularVehicleSimulationCU.cpp
+++ b/Engine/Plugins/Experimental/ChaosModularVehicle/Source/ChaosModularVehicleEngine/Private/ChaosModularVehicle/ModularVehicleSimulationCU.cpp
@@ -214,6 +214,23 @@ void FModularVehicleSimulation::PerformAdditionalSimWork(UWorld* InWorld, const
 	check(Proxy);
 	Chaos::EnsureIsInPhysicsThreadContext();
 
+	// clear all ground interactions
+	if (SimModuleTree)
+	{
+		const TArray<FSimModuleTree::FSimModuleNode>& ModuleArray = SimModuleTree->GetSimulationModuleTree();
+
+		for (const FSimModuleTree::FSimModuleNode& Node : ModuleArray)
+		{
+			if (Node.IsValid() && Node.SimModule && Node.SimModule->IsEnabled())
+			{
+				if (FWheelBaseInterface* Wheel = Node.SimModule->Cast<FWheelBaseInterface>())
+				{
+					Wheel->ClearGroundBody();
+				}
+			}
+		}
+	}
+
 	if (SimModuleTree && RootParticle)
 	{
 		FRigidTransform3 ClusterWorldTM = FRigidTransform3(RootParticle->GetX(), RootParticle->GetR());
diff --git a/Engine/Source/Runtime/Experimental/ChaosVehicles/ChaosVehiclesCore/Public/SimModule/TorqueSimModule.h b/Engine/Source/Runtime/Experimental/ChaosVehicles/ChaosVehiclesCore/Public/SimModule/TorqueSimModule.h
--- a/Engine/Source/Runtime/Experimental/ChaosVehicles/ChaosVehiclesCore/Public/SimModule/TorqueSimModule.h
+++ b/Engine/Source/Runtime/Experimental/ChaosVehicles/ChaosVehiclesCore/Public/SimModule/TorqueSimModule.h
@@ -145,6 +145,7 @@ namespace Chaos
 		}
 
 		Chaos::FPBDRigidParticleHandle* GetGroundBody() const { return GroundBody; }
+		void ClearGroundBody() { GroundBody = nullptr; }
 		FVector GetImpactPoint() const { return ImpactPoint; }
 		FVector GetImpactNormal() const { return ImpactNormal; }
 		virtual float GetWheelRadius() const = 0;
