[Chaos Modular Vehicles]

Fix bug in output dispatch that sent the wrong output data to most vehicles.

#Synced-CL 50695105

diff --git a/Engine/Plugins/Experimental/ChaosModularVehicle/Source/ChaosModularVehicleEngine/Private/ChaosModularVehicle/ChaosSimModuleManager.cpp b/Engine/Plugins/Experimental/ChaosModularVehicle/Source/ChaosModularVehicleEngine/Private/ChaosModularVehicle/ChaosSimModuleManager.cpp
--- a/Engine/Plugins/Experimental/ChaosModularVehicle/Source/ChaosModularVehicleEngine/Private/ChaosModularVehicle/ChaosSimModuleManager.cpp
+++ b/Engine/Plugins/Experimental/ChaosModularVehicle/Source/ChaosModularVehicleEngine/Private/ChaosModularVehicle/ChaosSimModuleManager.cpp
@@ -425,7 +425,7 @@ void FChaosSimModuleManager::ParallelUpdateVehicles()
 			TStrongObjectPtr<UModularVehicleBaseComponent> StrongPtr = Vehicle.Pin();
 			if (StrongPtr.IsValid())
 			{
-				StrongPtr->SetCurrentAsyncData(AsyncInput->VehicleInputs.Num(), PreviousOutput, NextOutput, Alpha, Timestamp);
+				StrongPtr->SetCurrentAsyncData(PreviousOutput, NextOutput, Alpha, Timestamp);
 			}
 		}
 	}
diff --git a/Engine/Plugins/Experimental/ChaosModularVehicle/Source/ChaosModularVehicleEngine/Private/ChaosModularVehicle/ChaosSimModuleManagerAsyncCallback.cpp b/Engine/Plugins/Experimental/ChaosModularVehicle/Source/ChaosModularVehicleEngine/Private/ChaosModularVehicle/ChaosSimModuleManagerAsyncCallback.cpp
--- a/Engine/Plugins/Experimental/ChaosModularVehicle/Source/ChaosModularVehicleEngine/Private/ChaosModularVehicle/ChaosSimModuleManagerAsyncCallback.cpp
+++ b/Engine/Plugins/Experimental/ChaosModularVehicle/Source/ChaosModularVehicleEngine/Private/ChaosModularVehicle/ChaosSimModuleManagerAsyncCallback.cpp
@@ -112,7 +112,7 @@ void FChaosSimModuleManagerAsyncCallback::OnPreSimulate_Internal()
 
 		bool bWake = false;
 		OutputVehiclesBatch[Idx] = VehicleInput.Simulate(World, DeltaTime, SimTime, bWake);
-
+		OutputVehiclesBatch[Idx]->Vehicle = VehicleInput.GetVehicle();
 	};
 
 	bool ForceSingleThread = !GSimModuleDebugParams.EnableMultithreading;
diff --git a/Engine/Plugins/Experimental/ChaosModularVehicle/Source/ChaosModularVehicleEngine/Private/ChaosModularVehicle/ModularVehicleBaseComponent.cpp b/Engine/Plugins/Experimental/ChaosModularVehicle/Source/ChaosModularVehicleEngine/Private/ChaosModularVehicle/ModularVehicleBaseComponent.cpp
--- a/Engine/Plugins/Experimental/ChaosModularVehicle/Source/ChaosModularVehicleEngine/Private/ChaosModularVehicle/ModularVehicleBaseComponent.cpp
+++ b/Engine/Plugins/Experimental/ChaosModularVehicle/Source/ChaosModularVehicleEngine/Private/ChaosModularVehicle/ModularVehicleBaseComponent.cpp
@@ -870,15 +870,15 @@ void UModularVehicleBaseComponent::ServerUpdateState_Implementation(const FModul
 }
 
 
-void UModularVehicleBaseComponent::SetCurrentAsyncData(int32 InputIdx, FChaosSimModuleManagerAsyncOutput* CurOutput, FChaosSimModuleManagerAsyncOutput* NextOutput, float Alpha, int32 VehicleManagerTimestamp)
+void UModularVehicleBaseComponent::SetCurrentAsyncData(FChaosSimModuleManagerAsyncOutput* CurOutput, FChaosSimModuleManagerAsyncOutput* NextOutput, float Alpha, int32 VehicleManagerTimestamp)
 {
-	SetCurrentAsyncDataInternal(InputIdx, CurOutput, NextOutput, Alpha, VehicleManagerTimestamp);
+	SetCurrentAsyncDataInternal(CurOutput, NextOutput, Alpha, VehicleManagerTimestamp);
 }
 
 /************************************************************************/
 /* Setup the current async I/O data                                     */
 /************************************************************************/
-void UModularVehicleBaseComponent::SetCurrentAsyncDataInternal(int32 InputIdx, FChaosSimModuleManagerAsyncOutput* CurOutput, FChaosSimModuleManagerAsyncOutput* NextOutput, float Alpha, int32 VehicleManagerTimestamp)
+void UModularVehicleBaseComponent::SetCurrentAsyncDataInternal(FChaosSimModuleManagerAsyncOutput* CurOutput, FChaosSimModuleManagerAsyncOutput* NextOutput, float Alpha, int32 VehicleManagerTimestamp)
 {
 	ensure(CurAsyncOutput == nullptr);	//should get reset after update is done
 
@@ -888,68 +888,38 @@ void UModularVehicleBaseComponent::SetCurrentAsyncDataInternal(int32 InputIdx, F
 	// We need to find our vehicle in the output given
 	if (CurOutput)
 	{
-		for (int32 PendingOutputIdx = 0; PendingOutputIdx < OutputsWaitingOn.Num(); ++PendingOutputIdx)
+		for (int32 CurOutputIdx = 0; CurOutputIdx < CurOutput->VehicleOutputs.Num(); ++CurOutputIdx)
 		{
-			// Found the correct pending output, use index to get the vehicle.
-			if (OutputsWaitingOn[PendingOutputIdx].Timestamp == CurOutput->Timestamp)
+			FModularVehicleAsyncOutput* VehicleOutput = CurOutput->VehicleOutputs[CurOutputIdx].Get();
+			if (VehicleOutput && VehicleOutput->bValid && VehicleOutput->Type == CurAsyncType)
 			{
-				const int32 VehicleIdx = OutputsWaitingOn[PendingOutputIdx].Idx;
-				FModularVehicleAsyncOutput* VehicleOutput = CurOutput->VehicleOutputs[VehicleIdx].Get();
-				if (VehicleOutput && VehicleOutput->bValid && VehicleOutput->Type == CurAsyncType)
+				if (VehicleOutput->Vehicle == this)
 				{
+					// Found the correct vehicle output.
 					CurAsyncOutput = VehicleOutput;
-
-					if (NextOutput && NextOutput->Timestamp == CurOutput->Timestamp)
-					{
-						// This can occur when sub-stepping - in this case, VehicleOutputs will be in the same order in NextOutput and CurOutput.
-						FModularVehicleAsyncOutput* VehicleNextOutput = NextOutput->VehicleOutputs[VehicleIdx].Get();
-						if (VehicleNextOutput && VehicleNextOutput->bValid && VehicleNextOutput->Type == CurAsyncType)
-						{
-							NextAsyncOutput = VehicleNextOutput;
-							OutputInterpAlpha = Alpha;
-						}
-					}
-				}
-
-				// these are sorted by timestamp, we are using latest, so remove entries that came before it.
-				TArray<FAsyncOutputWrapper> NewOutputsWaitingOn;
-				for (int32 CopyIndex = PendingOutputIdx; CopyIndex < OutputsWaitingOn.Num(); ++CopyIndex)
-				{
-					NewOutputsWaitingOn.Add(OutputsWaitingOn[CopyIndex]);
+					break;
 				}
-
-				OutputsWaitingOn = MoveTemp(NewOutputsWaitingOn);
-				break;
 			}
 		}
 
 	}
 
-	if (NextOutput && CurOutput)
+	if(NextOutput)
 	{
-		if (NextOutput->Timestamp != CurOutput->Timestamp)
+		for (int32 NextOutputIdx = 0; NextOutputIdx < NextOutput->VehicleOutputs.Num(); ++NextOutputIdx)
 		{
-			// NextOutput and CurOutput occurred in different steps, so we need to search for our specific vehicle.
-			for (int32 PendingOutputIdx = 0; PendingOutputIdx < OutputsWaitingOn.Num(); ++PendingOutputIdx)
+			FModularVehicleAsyncOutput* VehicleNextOutput = NextOutput->VehicleOutputs[NextOutputIdx].Get();
+			if (VehicleNextOutput && VehicleNextOutput->bValid && VehicleNextOutput->Type == CurAsyncType)
 			{
-				// Found the correct pending output, use index to get the vehicle.
-				if (OutputsWaitingOn[PendingOutputIdx].Timestamp == NextOutput->Timestamp)
+				if (VehicleNextOutput->Vehicle == this)
 				{
-					FModularVehicleAsyncOutput* VehicleOutput = NextOutput->VehicleOutputs[OutputsWaitingOn[PendingOutputIdx].Idx].Get();
-					if (VehicleOutput && VehicleOutput->bValid && VehicleOutput->Type == CurAsyncType)
-					{
-						NextAsyncOutput = VehicleOutput;
-						OutputInterpAlpha = Alpha;
-					}
+					NextAsyncOutput = VehicleNextOutput;
+					OutputInterpAlpha = Alpha;
 					break;
 				}
 			}
 		}
 	}
-
-	FAsyncOutputWrapper& NewOutput = OutputsWaitingOn.AddDefaulted_GetRef();
-	NewOutput.Timestamp = VehicleManagerTimestamp;
-	NewOutput.Idx = InputIdx;
 }
 
 IPhysicsProxyBase* UModularVehicleBaseComponent::GetPhysicsProxy() const 
diff --git a/Engine/Plugins/Experimental/ChaosModularVehicle/Source/ChaosModularVehicleEngine/Public/ChaosModularVehicle/ChaosSimModuleManagerAsyncCallback.h b/Engine/Plugins/Experimental/ChaosModularVehicle/Source/ChaosModularVehicleEngine/Public/ChaosModularVehicle/ChaosSimModuleManagerAsyncCallback.h
--- a/Engine/Plugins/Experimental/ChaosModularVehicle/Source/ChaosModularVehicleEngine/Public/ChaosModularVehicle/ChaosSimModuleManagerAsyncCallback.h
+++ b/Engine/Plugins/Experimental/ChaosModularVehicle/Source/ChaosModularVehicleEngine/Public/ChaosModularVehicle/ChaosSimModuleManagerAsyncCallback.h
@@ -274,7 +274,7 @@ struct FModularVehicleAsyncInput
 	UE_API virtual void ProcessInputs();
 
 	void SetVehicle(UModularVehicleBaseComponent* VehicleIn) { Vehicle = VehicleIn; }
-	UModularVehicleBaseComponent* GetVehicle() { return Vehicle; }
+	UModularVehicleBaseComponent* GetVehicle() const { return Vehicle; }
 
 	const EChaosAsyncVehicleDataType Type;
 	IPhysicsProxyBase* Proxy;
@@ -309,10 +309,12 @@ struct FModularVehicleAsyncOutput
 	const EChaosAsyncVehicleDataType Type;
 	bool bValid;	// indicates no work was done
 	FPhysicsVehicleOutput VehicleSimOutput;
+	UModularVehicleBaseComponent* Vehicle;
 
 	FModularVehicleAsyncOutput(EChaosAsyncVehicleDataType InType = EChaosAsyncVehicleDataType::AsyncInvalid)
 		: Type(InType)
 		, bValid(false)
+		, Vehicle(nullptr)
 	{ }
 
 	virtual ~FModularVehicleAsyncOutput()
diff --git a/Engine/Plugins/Experimental/ChaosModularVehicle/Source/ChaosModularVehicleEngine/Public/ChaosModularVehicle/ModularVehicleBaseComponent.h b/Engine/Plugins/Experimental/ChaosModularVehicle/Source/ChaosModularVehicleEngine/Public/ChaosModularVehicle/ModularVehicleBaseComponent.h
--- a/Engine/Plugins/Experimental/ChaosModularVehicle/Source/ChaosModularVehicleEngine/Public/ChaosModularVehicle/ModularVehicleBaseComponent.h
+++ b/Engine/Plugins/Experimental/ChaosModularVehicle/Source/ChaosModularVehicleEngine/Public/ChaosModularVehicle/ModularVehicleBaseComponent.h
@@ -152,7 +152,7 @@ public:
 
 	UE_API void PreTickGT(float DeltaTime);
 	UE_API void UpdateState(float DeltaTime);
-	UE_API void SetCurrentAsyncData(int32 InputIdx, FChaosSimModuleManagerAsyncOutput* CurOutput, FChaosSimModuleManagerAsyncOutput* NextOutput, float Alpha, int32 VehicleManagerTimestamp);
+	UE_API void SetCurrentAsyncData(FChaosSimModuleManagerAsyncOutput* CurOutput, FChaosSimModuleManagerAsyncOutput* NextOutput, float Alpha, int32 VehicleManagerTimestamp);
 
 	UE_API void ParallelUpdate(const Chaos::FCreatedModules& ModuleEvents);
 	UE_API void Update(float DeltaTime);
@@ -359,7 +359,7 @@ protected:
 
 	UE_API void ActionTreeUpdates(Chaos::FSimTreeUpdates* NextTreeUpdates);
 
-	UE_API void SetCurrentAsyncDataInternal(int32 InputIdx, FChaosSimModuleManagerAsyncOutput* CurOutput, FChaosSimModuleManagerAsyncOutput* NextOutput, float Alpha, int32 VehicleManagerTimestamp);
+	UE_API void SetCurrentAsyncDataInternal(FChaosSimModuleManagerAsyncOutput* CurOutput, FChaosSimModuleManagerAsyncOutput* NextOutput, float Alpha, int32 VehicleManagerTimestamp);
 	UE_API int32 FindParentsLastSimComponent(const USceneComponent* AttachedComponent);
 
 	UE_API IPhysicsProxyBase* GetPhysicsProxy() const;
@@ -444,7 +444,6 @@ public:
 		}
 	};
 
-	TArray<FAsyncOutputWrapper> OutputsWaitingOn;
 	TUniquePtr<FPhysicsVehicleOutput> PVehicleOutput;	/* physics simulation data output from the async physics thread */
 	TUniquePtr<FModularVehicleSimulation> VehicleSimulationPT;	/* simulation code running on the physics thread async callback */
 
