[Vehicles] Implement faster raycast solution for resims. ~2x performance improvement.

#Synced-CL 50894320
#Synced-CL 50894320

diff --git a/Engine/Plugins/Experimental/ChaosModularVehicle/Source/ChaosModularVehicleEngine/Private/ChaosModularVehicle/ModularVehicleSimulationCU.cpp b/Engine/Plugins/Experimental/ChaosModularVehicle/Source/ChaosModularVehicleEngine/Private/ChaosModularVehicle/ModularVehicleSimulationCU.cpp
--- a/Engine/Plugins/Experimental/ChaosModularVehicle/Source/ChaosModularVehicleEngine/Private/ChaosModularVehicle/ModularVehicleSimulationCU.cpp
+++ b/Engine/Plugins/Experimental/ChaosModularVehicle/Source/ChaosModularVehicleEngine/Private/ChaosModularVehicle/ModularVehicleSimulationCU.cpp
@@ -1,6 +1,7 @@
 // Copyright Epic Games, Inc. All Rights Reserved.
 
 #include "ChaosModularVehicle/ModularVehicleSimulationCU.h"
+
 #include "ChaosModularVehicle/ModularVehicleDefaultAsyncInput.h"
 #include "SimModule/SimModulesInclude.h"
 #include "SimModule/ModuleInput.h"
@@ -11,9 +12,14 @@
 #include "PhysicsEngine/PhysicsSettings.h"
 #include "Chaos/ClusterUnionManager.h"
 #include "Chaos/DebugDrawQueue.h"
+#include "Chaos/PhysicsObjectCollisionInterface.h"
+#include "Chaos/PhysicsObjectInternalInterface.h"
 #include "Engine/World.h"
 #include "Engine/HitResult.h"
 #include "ChaosModularVehicle/ModularVehicleDebug.h"
+#include "Components/PrimitiveComponent.h"
+#include "Engine/OverlapResult.h"
+#include "Runtime/Experimental/Chaos/Private/Chaos/PhysicsObjectInternal.h"
 
 FModularVehicleDebugParams GModularVehicleDebugParams;
 DEFINE_LOG_CATEGORY(LogModularVehicleSim);
@@ -32,7 +38,6 @@ FAutoConsoleVariableRef CVarChaosModularVehiclesFrictionOverride(TEXT("p.Modular
 FAutoConsoleVariableRef CVarChaosModularVehiclesDisableAnim(TEXT("p.ModularVehicle.DisableAnim"), GModularVehicleDebugParams.DisableAnim, TEXT("Disable animating wheels, etc"));
 #endif
 
-
 void FModularVehicleSimulation::Initialize(TUniquePtr<Chaos::FSimModuleTree>& InSimModuleTree)
 {
 	UE_LOG(LogModularVehicleSim, Log, TEXT("FModularVehicleSimulation::Initialize"));
@@ -214,6 +219,7 @@ void FModularVehicleSimulation::PerformAdditionalSimWork(UWorld* InWorld, const
 	check(Proxy);
 	Chaos::EnsureIsInPhysicsThreadContext();
 
+
 	// clear all ground interactions
 	if (SimModuleTree)
 	{
@@ -231,23 +237,126 @@ void FModularVehicleSimulation::PerformAdditionalSimWork(UWorld* InWorld, const
 		}
 	}
 
+	QUICK_SCOPE_CYCLE_COUNTER(STAT_VehicleSim_PerformAdditionalSimWork);
 	if (SimModuleTree && RootParticle)
 	{
 		FRigidTransform3 ClusterWorldTM = FRigidTransform3(RootParticle->GetX(), RootParticle->GetR());
-
+		AllInputs.VehicleWorldTransform = ClusterWorldTM;
+		
 		const TArray<FSimModuleTree::FSimModuleNode>& ModuleArray = SimModuleTree->GetSimulationModuleTree();
-
+		
+		
+		//Collect all rays, find intersecting objects via aabb box of rays.
+		FBox Box;
+		Box.Init();
 		for (const FSimModuleTree::FSimModuleNode& Node : ModuleArray)
 		{
 			if (Node.IsValid() && Node.SimModule && Node.SimModule->IsEnabled())
 			{
-				FRigidTransform3 Frame = Node.SimModule->GetParentRelativeTransform();
+				if (Node.SimModule->IsClustered() && Node.SimModule->IsBehaviourType(Chaos::eSimModuleTypeFlags::Raycast))
+				{
+					Chaos::FSpringTrace OutTrace;
+					Chaos::FSuspensionBaseInterface* Suspension = static_cast<Chaos::FSuspensionBaseInterface*>(Node.SimModule);
+					// would be cleaner an faster to just store radius in suspension also
+					float WheelRadius = 0;
+					if (Suspension->GetWheelSimTreeIndex() != Chaos::ISimulationModuleBase::INVALID_IDX)
+					{
+						Chaos::FWheelBaseInterface* Wheel = static_cast<Chaos::FWheelBaseInterface*>(ModuleArray[Suspension->GetWheelSimTreeIndex()].SimModule);
+						if (Wheel)
+						{
+							WheelRadius = Wheel->GetWheelRadius();
+						}
+					}
+					Suspension->GetWorldRaycastLocation(ClusterWorldTM, WheelRadius, OutTrace);
+					FVector TraceStart = OutTrace.Start;
+					FVector TraceEnd = OutTrace.End;
+					Box += TraceStart;
+					Box += TraceEnd;
+				}
+			}
+		}
+		FCollisionQueryParams& TraceParams = InputData.PhysicsInputs.TraceParams;
+		ECollisionChannel SpringCollisionChannel = InputData.PhysicsInputs.CollisionChannel;
+		FCollisionResponseParams ResponseParams = InputData.PhysicsInputs.TraceCollisionResponse;
+		TraceParams.bIgnoreTouches = true;
+		TraceParams.bSkipNarrowPhase = true;
+		TraceParams.bTraceComplex = true;
+		
+		TArray<FOverlapResult> Overlaps;
+		bool bHasOverlap = false;
+		if (InWorld && Box.IsValid && !Box.ContainsNaN() && Box.GetSize().Length() > 0)
+		{
+			FCollisionShape Shape = FCollisionShape::MakeBox(Box.GetExtent());
+			bHasOverlap = FGenericPhysicsInterface::GeomOverlapMulti(
+				InWorld,
+				Shape,
+				Box.GetCenter(),
+				FQuat::Identity,
+				Overlaps,
+				SpringCollisionChannel,
+				TraceParams,
+				ResponseParams,
+				FCollisionObjectQueryParams::DefaultObjectQueryParam);
+		}
+
+		TArray<FConstPhysicsObjectHandle> BlockingPhysicsObjects;
+		TArray<FOverlapResult> BlockingOverlaps;
+		if(bHasOverlap)
+		{
+			for(FOverlapResult& Overlap : Overlaps)
+			{
+				if (Overlap.PhysicsObject && Overlap.bBlockingHit)
+				{
+					if (TThreadParticle<EThreadContext::Internal>* Part = Overlap.PhysicsObject->GetParticle<EThreadContext::Internal>())
+					{
+						BlockingPhysicsObjects.Add(Overlap.PhysicsObject);
+						BlockingOverlaps.Add(Overlap);;
+					}
+				}
+			}
+		}
+		FReadPhysicsObjectInterface_Internal Interface = FPhysicsObjectInternalInterface::GetRead();
+		FPhysicsObjectCollisionInterface_Internal CollisionInterface{ Interface };
+		
+		auto GetMaterialFromInternalFaceIndex_InternalHelper = [](const FPhysicsShape& Shape, const Chaos::FGeometryParticleHandle& PTActor, uint32 InternalFaceIndex)-> Chaos::FChaosPhysicsMaterial*
+		{
+			const auto& Materials = Shape.GetMaterials();
+			if(Materials.Num() > 0 && PTActor.PhysicsProxy())
+			{
+				Chaos::FPBDRigidsSolver* Solver = PTActor.PhysicsProxy()->GetSolver<Chaos::FPBDRigidsSolver>();
+
+				if(ensure(Solver))
+				{
+					if(Materials.Num() == 1)
+					{
+						return Solver->GetSimMaterials().Get(Materials[0].InnerHandle);
+					}
 
-						
-				AllInputs.VehicleWorldTransform = ClusterWorldTM;
+					uint8 Index = Shape.GetGeometry()->GetMaterialIndex(InternalFaceIndex);
 
+					if(Materials.IsValidIndex(Index))
+					{
+						return Solver->GetSimMaterials().Get(Materials[Index].InnerHandle);
+					}
+				}
+			}
+
+			return nullptr;
+		};
+		
+		auto GetUserDataHelper = [](const Chaos::FChaosPhysicsMaterial& Material) -> UPhysicalMaterial*
+		{
+			void* UserData = Material.UserData;
+			return UserData ? FChaosUserData::Get<UPhysicalMaterial>(UserData) : nullptr;
+		};
+		
+		for (const FSimModuleTree::FSimModuleNode& Node : ModuleArray)
+		{
+			if (Node.IsValid() && Node.SimModule && Node.SimModule->IsEnabled())
+			{
 				if (Node.SimModule->IsClustered() && Node.SimModule->IsBehaviourType(Chaos::eSimModuleTypeFlags::Raycast))
 				{
+					QUICK_SCOPE_CYCLE_COUNTER(STAT_VehicleSim_SuspenssionRaycast);
 					Chaos::FSpringTrace OutTrace;
 					Chaos::FSuspensionBaseInterface* Suspension = static_cast<Chaos::FSuspensionBaseInterface*>(Node.SimModule);
 
@@ -267,50 +376,91 @@ void FModularVehicleSimulation::PerformAdditionalSimWork(UWorld* InWorld, const
 					FVector TraceStart = OutTrace.Start;
 					FVector TraceEnd = OutTrace.End;
 
-					const FCollisionQueryParams& TraceParams = InputData.PhysicsInputs.TraceParams;
 					FVector TraceVector(TraceStart - TraceEnd);
 					FVector TraceNormal = TraceVector.GetSafeNormal();
-
-					TArray<FHitResult> HitResultsOut;
-					FHitResult HitResult = FHitResult();
-					ECollisionChannel SpringCollisionChannel = InputData.PhysicsInputs.CollisionChannel;
-					const FCollisionResponseParams& ResponseParams = InputData.PhysicsInputs.TraceCollisionResponse;
-					if (InWorld)
+					
+					ChaosInterface::FPTLocationHit HitResult = ChaosInterface::FPTLocationHit();
+					bool bHasActualHit = false;
+					if (InWorld && bHasOverlap && BlockingPhysicsObjects.Num() > 0)
 					{
 						switch (InputData.PhysicsInputs.TraceType)
 						{
 							case ETraceType::Spherecast:
-							{
-								float QueryRadius = WheelRadius;
-								if (Chaos::Private::FGenericPhysicsInterface_Internal::SpherecastMulti(InWorld
-									, QueryRadius
-									, HitResultsOut
-									, TraceStart + TraceNormal * WheelRadius
-									, TraceEnd + TraceNormal * WheelRadius
-									, SpringCollisionChannel
-									, TraceParams
-									, ResponseParams))
 								{
-									HitResult = HitResultsOut.Last();
+									QUICK_SCOPE_CYCLE_COUNTER(STAT_VehicleSim_ShapeSweep);
+									ChaosInterface::FPTSweepHit SweepResult = ChaosInterface::FPTSweepHit();
+									FSweepParameters SweepParams;
+									SweepParams.bComputeMTD = false;
+									SweepParams.bSweepComplex = TraceParams.bTraceComplex;
+									FTransform SphereStart(FQuat::Identity,TraceStart);
+									const FPhysicsShapeAdapter SphereShapeAdapter(FQuat::Identity, FCollisionShape::MakeSphere(WheelRadius));
+									bHasActualHit = CollisionInterface.ShapeSweep(BlockingPhysicsObjects, SphereShapeAdapter.GetGeometry(), SphereStart, TraceEnd, SweepParams, SweepResult);
+									if (bHasActualHit)
+									{
+										HitResult.WorldPosition = SweepResult.WorldPosition;
+										HitResult.WorldNormal = SweepResult.WorldNormal;
+										HitResult.Distance = SweepResult.Distance;
+										HitResult.Actor = SweepResult.Actor;
+										HitResult.Shape = SweepResult.Shape;
+										HitResult.FaceIndex = SweepResult.FaceIndex;
+										HitResult.ElementIndex = SweepResult.ElementIndex;
+										HitResult.Flags	= SweepResult.Flags;
+										HitResult.FaceNormal = SweepResult.FaceNormal;
+									}
 								}
-							}
 							break;
-
 							case ETraceType::Raycast:
 							default:
+								{
+									ChaosInterface::FPTRaycastHit RayResult = ChaosInterface::FPTRaycastHit();
+									QUICK_SCOPE_CYCLE_COUNTER(STAT_VehicleSim_LineTrace);
+									bHasActualHit = CollisionInterface.LineTrace(BlockingPhysicsObjects, TraceStart, TraceEnd, TraceParams.bTraceComplex, RayResult);
+									if (bHasActualHit)
+									{
+										HitResult.WorldPosition = RayResult.WorldPosition;
+										HitResult.WorldNormal = RayResult.WorldNormal;
+										HitResult.Distance = RayResult.Distance;
+										HitResult.Actor = RayResult.Actor;
+										HitResult.Shape = RayResult.Shape;
+										HitResult.FaceIndex = RayResult.FaceIndex;
+										HitResult.ElementIndex = RayResult.ElementIndex;
+										HitResult.Flags	= RayResult.Flags;
+										HitResult.FaceNormal = RayResult.FaceNormal;
+									}
+								}
+							break;
+						}
+					}
+					
+					FOverlapResult ActualHitOverlapResult;
+					FConstPhysicsObjectHandle ActualHitObjectHandle;
+					if (bHasActualHit)
+					{
+						if (!HitResult.Actor)
+						{
+							bHasActualHit = false;
+						}
+						else
+						{
+							for (int32 Idx = 0; Idx < BlockingPhysicsObjects.Num(); ++Idx)
 							{
-								float QueryRadius = 0.0f;
-								if (Chaos::Private::FGenericPhysicsInterface_Internal::SpherecastMulti(InWorld, QueryRadius, HitResultsOut, TraceStart, TraceEnd, SpringCollisionChannel, TraceParams, ResponseParams))
+							
+								if (HitResult.Actor->PhysicsProxy() == BlockingPhysicsObjects[Idx]->PhysicsProxy())
 								{
-									HitResult = HitResultsOut.Last();
+									ActualHitOverlapResult = BlockingOverlaps[Idx];
+									ActualHitObjectHandle = BlockingPhysicsObjects[Idx];
+									break;
 								}
 							}
-							break;
+							if (!ActualHitOverlapResult.bBlockingHit)
+							{
+								bHasActualHit = false;
+							}
 						}
 					}
 
 					float Offset = Suspension->GetMaxSpringLength();
-					if (HitResult.bBlockingHit && GModularVehicleDebugParams.SuspensionRaycastsEnabled)
+					if (bHasActualHit && GModularVehicleDebugParams.SuspensionRaycastsEnabled)
 					{
 						Offset = HitResult.Distance - WheelRadius;
 
@@ -327,36 +477,38 @@ void FModularVehicleSimulation::PerformAdditionalSimWork(UWorld* InWorld, const
 								}
 								else
 								{
-									if (HitResult.PhysMaterial.IsValid())
+									if (HitResult.Shape && HitResult.Actor && HitResult.Actor->PhysicsProxy() && HitResult.FaceIndex >= 0)
 									{
-										Wheel->SetSurfaceFriction(HitResult.PhysMaterial->Friction);
+										if (const FPhysicsMaterial* Material = GetMaterialFromInternalFaceIndex_InternalHelper(*HitResult.Shape, *HitResult.Actor, HitResult.FaceIndex))
+										{
+											Wheel->SetSurfaceFriction(Material->Friction);
+										}
 									}
 								}
 
+								
 								IPhysicsProxyBase* HitProxy = nullptr;
 								if (InWorld)
 								{
-									if (Chaos::FPhysicsObjectHandle PhysicsObject = HitResult.PhysicsObject)
-									{
-										HitProxy = Chaos::FPhysicsObjectInterface::GetProxy({ &PhysicsObject, 1 });
-									}
-
+									HitProxy = Chaos::FPhysicsObjectInterface::GetProxy({ &ActualHitObjectHandle, 1 });
+									
 									if (HitProxy && HitProxy->GetType() == EPhysicsProxyType::SingleParticleProxy)
 									{
 										if (FSingleParticlePhysicsProxy* ParticleProxy = static_cast<FSingleParticlePhysicsProxy*>(HitProxy))
 										{
-											Chaos::FPBDRigidParticleHandle* GroundParticle = ParticleProxy->GetHandle_LowLevel() ? ParticleProxy->GetHandle_LowLevel()->CastToRigidParticle() : nullptr;
-											Wheel->SetGroundInteraction(GroundParticle, HitResult.ImpactPoint, HitResult.ImpactNormal);
+											Chaos::FPBDRigidParticleHandle* GroundParticle = HitResult.Actor->CastToRigidParticle(); 
+											Wheel->SetGroundInteraction(GroundParticle, HitResult.WorldPosition, HitResult.WorldNormal);
 										}
 									}
 								}
+								 
 							}
 						}
 
 #if CHAOS_DEBUG_DRAW
 						if (GModularVehicleDebugParams.ShowSuspensionRaycasts)
 						{
-							Chaos::FDebugDrawQueue::GetInstance().DrawDebugSphere(HitResult.ImpactPoint, 3, 16, FColor::Red, false, -1.f, 0, 10.f);
+							Chaos::FDebugDrawQueue::GetInstance().DrawDebugSphere(HitResult.WorldPosition, 3, 16, FColor::Red, false, -1.f, 0, 10.f);
 						}
 
 						if (Suspension->GetWheelSimTreeIndex() != Chaos::ISimulationModuleBase::INVALID_IDX)
@@ -379,7 +531,7 @@ void FModularVehicleSimulation::PerformAdditionalSimWork(UWorld* InWorld, const
 											Col = FColor::Red;
 										}
 									}
-									Chaos::FDebugDrawQueue::GetInstance().DrawDebugString(HitResult.ImpactPoint + FVec3(0, 50, 50), TextOut, nullptr, Col, -1.f, true, 1.0f);
+									Chaos::FDebugDrawQueue::GetInstance().DrawDebugString(HitResult.WorldNormal + FVec3(0, 50, 50), TextOut, nullptr, Col, -1.f, true, 1.0f);
 								}
 							}
 						}
@@ -391,11 +543,11 @@ void FModularVehicleSimulation::PerformAdditionalSimWork(UWorld* InWorld, const
 					if (GModularVehicleDebugParams.ShowSuspensionRaycasts)
 					{
 						FColor DrawColor = FColor::Green;
-						DrawColor = (HitResult.bBlockingHit) ? FColor::Red : FColor::Green;
+						DrawColor = (bHasActualHit) ? FColor::Red : FColor::Green;
 						Chaos::FDebugDrawQueue::GetInstance().DrawDebugLine(TraceStart, TraceEnd, DrawColor, false, -1.f, 0, 2.f);
 						Chaos::FDebugDrawQueue::GetInstance().DrawDebugSphere(TraceStart, 3, 16, FColor::White, false, -1.f, 0, 10.f);
-						Chaos::FDebugDrawQueue::GetInstance().DrawDebugSphere(HitResult.ImpactPoint, 1, 16, FColor::Red, false, -1.f, 0, 10.f);
-						FString TextOut = FString::Format(TEXT("{0}"), { HitResult.Time });
+						Chaos::FDebugDrawQueue::GetInstance().DrawDebugSphere(HitResult.WorldPosition, 1, 16, FColor::Red, false, -1.f, 0, 10.f);
+						FString TextOut = FString::Format(TEXT("{0}"), { HitResult.Distance });
 
 						FColor Col = FColor::White;
 						if (InWorld)
@@ -409,21 +561,27 @@ void FModularVehicleSimulation::PerformAdditionalSimWork(UWorld* InWorld, const
 								Col = FColor::Red;
 							}
 						}
-						Chaos::FDebugDrawQueue::GetInstance().DrawDebugString(HitResult.ImpactPoint + FVec3(0, 50, 50), TextOut, nullptr, Col, -1.f, true, 1.0f);
+						Chaos::FDebugDrawQueue::GetInstance().DrawDebugString(HitResult.WorldPosition + FVector(0, 50, 50), TextOut, nullptr, Col, -1.f, true, 1.0f);
 					}
 
 					if (GModularVehicleDebugParams.ShowRaycastMaterial)
 					{
-						if (HitResult.PhysMaterial.IsValid())
+						if (HitResult.Shape && HitResult.Actor && HitResult.Actor->PhysicsProxy() && HitResult.FaceIndex >= 0)
 						{
-							FDebugDrawQueue::GetInstance().DrawDebugString(HitResult.ImpactPoint, HitResult.PhysMaterial->GetName(), nullptr, FColor::White, -1.f, true, 1.0f);
+							if (const FPhysicsMaterial* Material = GetMaterialFromInternalFaceIndex_InternalHelper(*HitResult.Shape, *HitResult.Actor, HitResult.FaceIndex))
+							{
+								if (UPhysicalMaterial* GTMaterial = GetUserDataHelper(*Material))
+								{
+									FDebugDrawQueue::GetInstance().DrawDebugString(HitResult.WorldPosition, GTMaterial->GetName(), nullptr, FColor::White, -1.f, true, 1.0f);
+								}
+							}
 						}
 					}
 
 					if (GModularVehicleDebugParams.ShowWheelCollisionNormal)
 					{
-						FVector Pt = HitResult.ImpactPoint;
-						FDebugDrawQueue::GetInstance().DrawDebugLine(Pt, Pt + HitResult.Normal * 20.0f, FColor::Yellow, false, 1.0f, 0, 1.0f);
+						FVector Pt = HitResult.WorldPosition;
+						FDebugDrawQueue::GetInstance().DrawDebugLine(Pt, Pt + HitResult.WorldNormal * 20.0f, FColor::Yellow, false, 1.0f, 0, 1.0f);
 						FDebugDrawQueue::GetInstance().DrawDebugSphere(Pt, 5.0f, 4, FColor::White, false, 1.0f, 0, 1.0f);
 					}
 
@@ -435,29 +593,34 @@ void FModularVehicleSimulation::PerformAdditionalSimWork(UWorld* InWorld, const
 					float HitDistance = 0.f;
 					if (InputData.PhysicsInputs.TraceType == ETraceType::Spherecast)
 					{
-						HitPoint = HitResult.Location;
+						HitPoint = HitResult.WorldPosition;
 						HitDistance = HitResult.Distance;
 					}
 					else
 					{
-						HitPoint = HitResult.ImpactPoint + Up * WheelRadius;
+						HitPoint = HitResult.WorldPosition + Up * WheelRadius;
 						HitDistance = HitResult.Distance - WheelRadius;
 					}
-
-					IPhysicsProxyBase* HitProxy = nullptr;
-					if (Chaos::FPhysicsObjectHandle PhysicsObject = HitResult.PhysicsObject)
+					
+					TEnumAsByte<EPhysicalSurface> SurfaceType = EPhysicalSurface::SurfaceType_Default;
+					if (HitResult.Shape && HitResult.Actor && HitResult.FaceIndex >= 0)
 					{
-						HitProxy = Chaos::FPhysicsObjectInterface::GetProxy({ &PhysicsObject, 1 });
+						if (const FPhysicsMaterial* Material = GetMaterialFromInternalFaceIndex_InternalHelper(*HitResult.Shape, *HitResult.Actor, HitResult.FaceIndex))
+						{
+							if (UPhysicalMaterial* GTMaterial = GetUserDataHelper(*Material))
+							{
+								SurfaceType = GTMaterial->SurfaceType;
+							}
+						}
 					}
-
-					const TEnumAsByte<EPhysicalSurface> DefaultSurfaceType = EPhysicalSurface::SurfaceType_Default;
+					
 					FSuspensionTargetPoint TargetPoint(
 						HitPoint
-						, HitResult.ImpactNormal
+						, HitResult.WorldNormal
 						, HitDistance
-						, HitResult.bBlockingHit
-						, HitResult.PhysMaterial.IsValid() ? HitResult.PhysMaterial->SurfaceType : DefaultSurfaceType
-						, HitProxy
+						, bHasActualHit
+						, SurfaceType
+						, bHasActualHit && HitResult.Actor ? HitResult.Actor->PhysicsProxy() : nullptr
 					);
 
 					Suspension->SetTargetPoint(TargetPoint);
