[Chaos][ModularVehicle] Fix suspension SpringLength and HitDistance set to 0 when not colliding. Now set at MaxSpringLength.

#Synced-CL 51012076

diff --git a/Engine/Plugins/Experimental/ChaosModularVehicle/Source/ChaosModularVehicleEngine/Private/ChaosModularVehicle/ModularVehicleSimulationCU.cpp b/Engine/Plugins/Experimental/ChaosModularVehicle/Source/ChaosModularVehicleEngine/Private/ChaosModularVehicle/ModularVehicleSimulationCU.cpp
--- a/Engine/Plugins/Experimental/ChaosModularVehicle/Source/ChaosModularVehicleEngine/Private/ChaosModularVehicle/ModularVehicleSimulationCU.cpp
+++ b/Engine/Plugins/Experimental/ChaosModularVehicle/Source/ChaosModularVehicleEngine/Private/ChaosModularVehicle/ModularVehicleSimulationCU.cpp
@@ -590,16 +590,14 @@ void FModularVehicleSimulation::PerformAdditionalSimWork(UWorld* InWorld, const
 					FVector Up = ClusterWorldTM.GetUnitAxis(EAxis::Z);
 
 					FVector HitPoint;
-					float HitDistance = 0.f;
+					float HitDistance = Offset;
 					if (InputData.PhysicsInputs.TraceType == ETraceType::Spherecast)
 					{
 						HitPoint = HitResult.WorldPosition;
-						HitDistance = HitResult.Distance;
 					}
 					else
 					{
 						HitPoint = HitResult.WorldPosition + Up * WheelRadius;
-						HitDistance = HitResult.Distance - WheelRadius;
 					}
 					
 					TEnumAsByte<EPhysicalSurface> SurfaceType = EPhysicalSurface::SurfaceType_Default;
