ModularVehicle: Recalculate suspension length on game thread so wheels don't look like they are going through the ground.

#Synced-CL 51657823

diff --git a/Engine/Plugins/Experimental/ChaosModularVehicle/Source/ChaosModularVehicleEngine/Private/ChaosModularVehicle/ModularVehicleBaseComponent.cpp b/Engine/Plugins/Experimental/ChaosModularVehicle/Source/ChaosModularVehicleEngine/Private/ChaosModularVehicle/ModularVehicleBaseComponent.cpp
--- a/Engine/Plugins/Experimental/ChaosModularVehicle/Source/ChaosModularVehicleEngine/Private/ChaosModularVehicle/ModularVehicleBaseComponent.cpp
+++ b/Engine/Plugins/Experimental/ChaosModularVehicle/Source/ChaosModularVehicleEngine/Private/ChaosModularVehicle/ModularVehicleBaseComponent.cpp
@@ -134,6 +134,15 @@ void UModularVehicleBaseComponent::ProduceInput(int32 PhysicsStep, int32 NumStep
 		AsyncInput->Proxy = GetPhysicsProxy();
 		AsyncInput->PhysicsInputs.NetworkInputs.VehicleInputs.KeepAwake = bKeepVehicleAwake;
 
+		AsyncInput->PhysicsInputs.CurrentTimeDilation = 1.0f;
+		if (GetWorld())
+		{
+			if (AWorldSettings* WorldSettings = GetWorld()->GetWorldSettings())
+			{
+				AsyncInput->PhysicsInputs.CurrentTimeDilation = FMath::Max(WorldSettings->GetEffectiveTimeDilation(), SMALL_NUMBER);
+			}
+		}
+
 		// All control inputs
 		AsyncInput->PhysicsInputs.NetworkInputs.VehicleInputs.Container = InputsContainer;
 
@@ -1032,27 +1041,35 @@ void UModularVehicleBaseComponent::ParallelUpdate(const Chaos::FCreatedModules&
 							}
 						}
 					}
-				}
 				
-				// extract/cache some generally useful values as we go as trying to locate this data later requires a search
-				if (PVehicleOutput->SimTreeOutputData[I]->IsSimType<Chaos::FTransmissionSimModule>())
-				{
-					// if there is more than one transmission then the last one will inform us of the current gear
-					CurrentGear = static_cast<Chaos::FTransmissionOutputData*>(PVehicleOutput->SimTreeOutputData[I])->CurrentGear;
-				}
-				else if (PVehicleOutput->SimTreeOutputData[I]->IsSimType<Chaos::FEngineSimModule>())
-				{
-					// if there is more than one engine then the last one will inform us of the engine RPM
-					Chaos::FEngineOutputData* Engine = static_cast<Chaos::FEngineOutputData*>(CurrentOutput->VehicleSimOutput.SimTreeOutputData[I]);
+					// extract/cache some generally useful values as we go as trying to locate this data later requires a search
+					if (PVehicleOutput->SimTreeOutputData[I]->IsSimType<Chaos::FTransmissionSimModule>())
+					{
+						// if there is more than one transmission then the last one will inform us of the current gear
+						CurrentGear = static_cast<Chaos::FTransmissionOutputData*>(PVehicleOutput->SimTreeOutputData[I])->CurrentGear;
+					}
+					else if (PVehicleOutput->SimTreeOutputData[I]->IsSimType<Chaos::FEngineSimModule>())
+					{
+						// if there is more than one engine then the last one will inform us of the engine RPM
+						Chaos::FEngineOutputData* Engine = static_cast<Chaos::FEngineOutputData*>(CurrentOutput->VehicleSimOutput.SimTreeOutputData[I]);
 
-					EngineRPM = Engine->RPM;
-					EngineTorque = Engine->Torque;
+						EngineRPM = Engine->RPM;
+						EngineTorque = Engine->Torque;
+					}
 				}
 
-
+				// Pass animation data to final location in ModuleAnimationSetups, which is used by the animation instance
 				if (Chaos::FSimOutputData* ModuleOutput = PVehicleOutput->SimTreeOutputData[I])
 				{
-					const Chaos::FSimModuleAnimationData& AnimData = ModuleOutput->AnimationData;
+					FTransform NewTransform = FTransform::Identity;
+					if (ClusterUnionComponent)
+					{
+						NewTransform = ClusterUnionComponent->GetComponentToWorld();
+					}
+
+					Chaos::FSimModuleAnimationData AnimData;
+					ModuleOutput->GetFinalAnimDataGameThread(NewTransform, AnimData);
+
 					int Index = AnimData.AnimationSetupIndex;
 					if ((Index >= 0) && (Index < ModuleAnimationSetups.Num()))
 					{
diff --git a/Engine/Plugins/Experimental/ChaosModularVehicle/Source/ChaosModularVehicleEngine/Private/ChaosModularVehicle/ModularVehicleSimulationCU.cpp b/Engine/Plugins/Experimental/ChaosModularVehicle/Source/ChaosModularVehicleEngine/Private/ChaosModularVehicle/ModularVehicleSimulationCU.cpp
--- a/Engine/Plugins/Experimental/ChaosModularVehicle/Source/ChaosModularVehicleEngine/Private/ChaosModularVehicle/ModularVehicleSimulationCU.cpp
+++ b/Engine/Plugins/Experimental/ChaosModularVehicle/Source/ChaosModularVehicleEngine/Private/ChaosModularVehicle/ModularVehicleSimulationCU.cpp
@@ -144,11 +144,12 @@ void FModularVehicleSimulation::SimulateModuleTree(UWorld* InWorld, float DeltaS
 		SimInputData.ControlInputs = &InputInterface;
 		SimInputData.StateInputs = &StateInterface;
 		SimInputData.bKeepVehicleAwake = InputData.PhysicsInputs.NetworkInputs.VehicleInputs.KeepAwake;
+		SimInputData.CurrentTimeDilation = InputData.PhysicsInputs.CurrentTimeDilation;
 
 
 		if(SimModuleTree->GetSimTreeProcessingOrder() != ESimTreeProcessingOrder::ManualOverride)
 		{
-			PerformAdditionalSimWork(InWorld, InputData, Proxy, SimInputData);
+			PerformAdditionalSimWork(DeltaSeconds, InWorld, InputData, Proxy, SimInputData);
 		}
 		// run the dynamics simulation, engine, suspension, wheels, aerofoils etc.
 		SimModuleTree->Simulate(DeltaSeconds, SimInputData, Proxy, RootParticle);
@@ -213,7 +214,7 @@ void FModularVehicleSimulation::CacheRootParticle(IPhysicsProxyBase* Proxy)
 	}
 }
 
-void FModularVehicleSimulation::PerformAdditionalSimWork(UWorld* InWorld, const FModularVehicleAsyncInput& InputData, IPhysicsProxyBase* Proxy, Chaos::FAllInputs& AllInputs)
+void FModularVehicleSimulation::PerformAdditionalSimWork(float DeltaSeconds, UWorld* InWorld, const FModularVehicleAsyncInput& InputData, IPhysicsProxyBase* Proxy, Chaos::FAllInputs& AllInputs)
 {
 	using namespace Chaos;
 	check(Proxy);
@@ -242,6 +243,7 @@ void FModularVehicleSimulation::PerformAdditionalSimWork(UWorld* InWorld, const
 	{
 		FRigidTransform3 ClusterWorldTM = FRigidTransform3(RootParticle->GetX(), RootParticle->GetR());
 		AllInputs.VehicleWorldTransform = ClusterWorldTM;
+		FVector ClusterVelocity = RootParticle->GetV();
 		
 		const TArray<FSimModuleTree::FSimModuleNode>& ModuleArray = SimModuleTree->GetSimulationModuleTree();
 		
@@ -267,7 +269,7 @@ void FModularVehicleSimulation::PerformAdditionalSimWork(UWorld* InWorld, const
 							WheelRadius = Wheel->GetWheelRadius();
 						}
 					}
-					Suspension->GetWorldRaycastLocation(ClusterWorldTM, WheelRadius, OutTrace);
+					Suspension->GetWorldTraceEndpoints(DeltaSeconds, ClusterWorldTM, ClusterVelocity, WheelRadius, OutTrace);
 					FVector TraceStart = OutTrace.Start;
 					FVector TraceEnd = OutTrace.End;
 					Box += TraceStart;
@@ -371,7 +373,7 @@ void FModularVehicleSimulation::PerformAdditionalSimWork(UWorld* InWorld, const
 						}
 					}
 
-					Suspension->GetWorldRaycastLocation(ClusterWorldTM, WheelRadius, OutTrace);
+					Suspension->GetWorldTraceEndpoints(DeltaSeconds, ClusterWorldTM, ClusterVelocity, WheelRadius, OutTrace);
 
 					FVector TraceStart = OutTrace.Start;
 					FVector TraceEnd = OutTrace.End;
diff --git a/Engine/Plugins/Experimental/ChaosModularVehicle/Source/ChaosModularVehicleEngine/Private/ChaosModularVehicle/SuspensionSimModule.cpp b/Engine/Plugins/Experimental/ChaosModularVehicle/Source/ChaosModularVehicleEngine/Private/ChaosModularVehicle/SuspensionSimModule.cpp
--- a/Engine/Plugins/Experimental/ChaosModularVehicle/Source/ChaosModularVehicleEngine/Private/ChaosModularVehicle/SuspensionSimModule.cpp
+++ b/Engine/Plugins/Experimental/ChaosModularVehicle/Source/ChaosModularVehicleEngine/Private/ChaosModularVehicle/SuspensionSimModule.cpp
@@ -42,16 +42,23 @@ using namespace Chaos;
 		SpringDisplacement = Setup().MaxLength - DisplacementInput;
 	}
 
-	void FSuspensionSimModule::GetWorldRaycastLocation(const FTransform& BodyTransform, float WheelRadius, FSpringTrace& OutTrace)
+	void FSuspensionSimModule::GetWorldTraceEndpoints(float DeltaSeconds, const FTransform& BodyTransform, const FVector& Velocity, float WheelRadius, FSpringTrace& OutTrace) const
 	{
 		FVector LocalDirection = Setup().SuspensionAxis;
 		FVector Local = GetParentRelativeTransform().GetLocation(); // change to just a vector and GetLocalLocation
 		FVector WorldLocation = BodyTransform.TransformPosition(Local);
 		FVector WorldDirection = BodyTransform.TransformVector(LocalDirection);
 
+		FVector MovementExpansion = FVector::ZeroVector;
+		// time dilation only an issue when time is slowed down
+		if (CurrentTimeDilation < 1.0f)
+		{
+			float Amount = FMath::Abs(WorldDirection.Dot(Velocity));
+			MovementExpansion += WorldDirection * (Amount * DeltaSeconds / CurrentTimeDilation);
+		}
+
 		OutTrace.Start = WorldLocation - WorldDirection * (Setup().MaxRaise);
-		OutTrace.End = WorldLocation + WorldDirection * (Setup().MaxDrop + WheelRadius);
-		float TraceLength = OutTrace.Start.Z - OutTrace.End.Z;
+		OutTrace.End = WorldLocation + WorldDirection * (Setup().MaxDrop + WheelRadius) + MovementExpansion;
 	}
 
 	void FSuspensionSimModule::OnConstruction_External(const Chaos::FPhysicsObjectHandle& PhysicsObject)
@@ -69,10 +76,12 @@ using namespace Chaos;
 	void FSuspensionSimModule::Simulate(float DeltaTime, const FAllInputs& Inputs, FSimModuleTree& VehicleModuleSystem)
 	{
 		{
+			CurrentTimeDilation = FMath::Max(Inputs.CurrentTimeDilation, SMALL_NUMBER);
 			float ForceIntoSurface = 0.0f;
 			if (SpringDisplacement > 0)
 			{
 				float Damping = Setup().SpringDamping;
+				DeltaTime = FMath::Max(DeltaTime, SMALL_NUMBER);
 				SpringSpeed = (LastDisplacement - SpringDisplacement) / DeltaTime;
 
 				float StiffnessForce = SpringDisplacement * Setup().SpringRate;
@@ -224,6 +233,14 @@ using namespace Chaos;
 			SpringDisplacementVector = -Sim->Setup().SuspensionAxis * Sim->SpringDisplacement + Sim->GetAnimationOffset();
 			SpringSpeed = Sim->SpringSpeed;
 			ImpactNormal = Sim->GetImpactNormal();
+
+			// Required to recalculate suspension length on the game thread (post physics integration) so it looks correct visually
+			ImpactPosition = Sim->GetTargetPosition();
+			ModuleLocalPosition = Sim->GetParentRelativeTransform().GetLocation();
+			LocalSuspensionAxis = Sim->Setup().SuspensionAxis;
+			MaxRaise = Sim->Setup().MaxRaise;
+			MaxDrop = Sim->Setup().MaxDrop;
+			bWheelHit = Sim->IsWheelInContact();
 		}
 	}
 
@@ -237,6 +254,36 @@ using namespace Chaos;
 		SpringDisplacement = FMath::Lerp(Current.SpringDisplacement, Next.SpringDisplacement, Alpha);
 		SpringSpeed = FMath::Lerp(Current.SpringSpeed, Next.SpringSpeed, Alpha);
 		ImpactNormal = FMath::Lerp(Current.ImpactNormal, Next.ImpactNormal, Alpha);
+
+		ImpactPosition = Next.ImpactPosition;
+		ModuleLocalPosition = FMath::Lerp(Current.ModuleLocalPosition, Next.ModuleLocalPosition, Alpha);
+		LocalSuspensionAxis = FMath::Lerp(Current.LocalSuspensionAxis, Next.LocalSuspensionAxis, Alpha);
+		MaxRaise = FMath::Lerp(Current.MaxRaise, Next.MaxRaise, Alpha);
+		MaxDrop = FMath::Lerp(Current.MaxDrop, Next.MaxDrop, Alpha);
+		bWheelHit = Next.bWheelHit;
+
+	}
+
+
+	void FSuspensionOutputData::GetFinalAnimDataGameThread(const FTransform& NewTransform, Chaos::FSimModuleAnimationData& AnimDataOut)
+	{
+		FSimOutputData::GetFinalAnimDataGameThread(NewTransform, AnimDataOut);
+
+		FVector WorldLocation = NewTransform.TransformPosition(ModuleLocalPosition);
+		FVector WorldDirection = NewTransform.TransformVector(LocalSuspensionAxis);
+		FVector NewStart = WorldLocation - WorldDirection * (MaxRaise);
+		FVector NewEnd = ImpactPosition;
+		float Length = FVector::DotProduct(NewStart - NewEnd, WorldDirection);
+		Length = FMath::Clamp(Length, -(MaxDrop+MaxRaise), 0);
+		FVector Movement = -LocalSuspensionAxis * (Length + MaxRaise);
+
+		if (!bWheelHit)
+		{
+			Movement = LocalSuspensionAxis * MaxDrop;
+		}
+
+		AnimDataOut.AnimFlags = EAnimationFlags::AnimatePosition;
+		AnimDataOut.AnimationLocOffset = Movement;
 	}
 
 #if !(UE_BUILD_SHIPPING || UE_BUILD_TEST)
diff --git a/Engine/Plugins/Experimental/ChaosModularVehicle/Source/ChaosModularVehicleEngine/Public/ChaosModularVehicle/ChaosSimModuleManagerAsyncCallback.h b/Engine/Plugins/Experimental/ChaosModularVehicle/Source/ChaosModularVehicleEngine/Public/ChaosModularVehicle/ChaosSimModuleManagerAsyncCallback.h
--- a/Engine/Plugins/Experimental/ChaosModularVehicle/Source/ChaosModularVehicleEngine/Public/ChaosModularVehicle/ChaosSimModuleManagerAsyncCallback.h
+++ b/Engine/Plugins/Experimental/ChaosModularVehicle/Source/ChaosModularVehicleEngine/Public/ChaosModularVehicle/ChaosSimModuleManagerAsyncCallback.h
@@ -240,6 +240,7 @@ struct FPhysicsModularVehicleInputs
 		, TraceParams()
 		, TraceCollisionResponse()
 		, TraceType(ETraceType::Raycast)
+		, CurrentTimeDilation(1.0f)
 	{
 	}
 	mutable FNetworkModularVehicleInputs NetworkInputs;
@@ -248,6 +249,7 @@ struct FPhysicsModularVehicleInputs
 	mutable FCollisionResponseContainer TraceCollisionResponse;
 	mutable ETraceType TraceType;
 	mutable FGameStateInputs StateInputs;
+	mutable float CurrentTimeDilation;
 };
 
 /**
diff --git a/Engine/Plugins/Experimental/ChaosModularVehicle/Source/ChaosModularVehicleEngine/Public/ChaosModularVehicle/ModularVehicleSimulationCU.h b/Engine/Plugins/Experimental/ChaosModularVehicle/Source/ChaosModularVehicleEngine/Public/ChaosModularVehicle/ModularVehicleSimulationCU.h
--- a/Engine/Plugins/Experimental/ChaosModularVehicle/Source/ChaosModularVehicleEngine/Public/ChaosModularVehicle/ModularVehicleSimulationCU.h
+++ b/Engine/Plugins/Experimental/ChaosModularVehicle/Source/ChaosModularVehicleEngine/Public/ChaosModularVehicle/ModularVehicleSimulationCU.h
@@ -101,7 +101,7 @@ public:
 
 	UE_API void ApplyDeferredForces(IPhysicsProxyBase* Proxy);
 
-	UE_API void PerformAdditionalSimWork(UWorld* InWorld, const FModularVehicleAsyncInput& InputData, IPhysicsProxyBase* Proxy, Chaos::FAllInputs& AllInputs);
+	UE_API void PerformAdditionalSimWork(float DeltaSeconds, UWorld* InWorld, const FModularVehicleAsyncInput& InputData, IPhysicsProxyBase* Proxy, Chaos::FAllInputs& AllInputs);
 
 	UE_API void FillOutputState(FModularVehicleAsyncOutput& Output);
 
diff --git a/Engine/Plugins/Experimental/ChaosModularVehicle/Source/ChaosModularVehicleEngine/Public/ChaosModularVehicle/SuspensionSimModule.h b/Engine/Plugins/Experimental/ChaosModularVehicle/Source/ChaosModularVehicleEngine/Public/ChaosModularVehicle/SuspensionSimModule.h
--- a/Engine/Plugins/Experimental/ChaosModularVehicle/Source/ChaosModularVehicleEngine/Public/ChaosModularVehicle/SuspensionSimModule.h
+++ b/Engine/Plugins/Experimental/ChaosModularVehicle/Source/ChaosModularVehicleEngine/Public/ChaosModularVehicle/SuspensionSimModule.h
@@ -56,14 +56,23 @@ struct FSuspensionOutputData
 	UE_API virtual void FillOutputState(const Chaos::ISimulationModuleBase* SimModule) override;
 	UE_API virtual void Lerp(const FSimOutputData& InCurrent, const Chaos::FSimOutputData& InNext, float Alpha) override;
 
+	UE_API virtual void GetFinalAnimDataGameThread(const FTransform& NewTransform, Chaos::FSimModuleAnimationData& AnimDataOut) override;
+
 #if !(UE_BUILD_SHIPPING || UE_BUILD_TEST)
 	UE_API virtual FString ToString() override;
 #endif
 
-	float SpringDisplacement;
-	FVector SpringDisplacementVector;
-	float SpringSpeed;
-	FVector ImpactNormal;
+	float SpringDisplacement = 0.0f;
+	FVector SpringDisplacementVector = FVector::ZeroVector;
+	float SpringSpeed = 0.0f;
+	FVector ImpactNormal = FVector(0.0f,0.0f,1.0f);
+
+	FVector ImpactPosition = FVector::ZeroVector;
+	FVector ModuleLocalPosition = FVector::ZeroVector;
+	FVector LocalSuspensionAxis = FVector::ZeroVector;
+	float MaxRaise = 0.0f;
+	float MaxDrop = 0.0f;
+	bool bWheelHit = false;
 };
 
 struct FSuspensionSettings
@@ -144,8 +153,7 @@ public:
 	virtual float GetMaxSpringLength() const override { return Setup().MaxLength; }
 	UE_API virtual float GetSpringLength() const override;
 	UE_API virtual  void SetSpringLength(float InLength, float WheelRadius) override;
-	UE_API virtual void GetWorldRaycastLocation(const FTransform& BodyTransform, float WheelRadius, Chaos::FSpringTrace& OutTrace) override;
-
+	UE_API virtual void GetWorldTraceEndpoints(float DeltaSeconds, const FTransform& BodyTransform, const FVector& Velocity, float WheelRadius, Chaos::FSpringTrace& OutTrace) const override;
 	UE_API virtual void OnConstruction_External(const Chaos::FPhysicsObjectHandle& PhysicsObject) override;
 	UE_API virtual void OnTermination_External() override;
 
@@ -163,10 +171,10 @@ protected:
 
 private:
 
-	float SpringDisplacement;
-	float LastDisplacement;
-	float SpringSpeed;
-
+	float SpringDisplacement = 0.0f;
+	float LastDisplacement = 0.0f;
+	float SpringSpeed = 0.0f;
+	float CurrentTimeDilation = 1.0f;
 	FPhysicsConstraintHandle ConstraintHandle;
 };
 
diff --git a/Engine/Source/Runtime/Experimental/ChaosVehicles/ChaosVehiclesCore/Public/SimModule/SimulationModuleBase.h b/Engine/Source/Runtime/Experimental/ChaosVehicles/ChaosVehiclesCore/Public/SimModule/SimulationModuleBase.h
--- a/Engine/Source/Runtime/Experimental/ChaosVehicles/ChaosVehiclesCore/Public/SimModule/SimulationModuleBase.h
+++ b/Engine/Source/Runtime/Experimental/ChaosVehicles/ChaosVehiclesCore/Public/SimModule/SimulationModuleBase.h
@@ -71,6 +71,7 @@ namespace Chaos
 		FInputInterface* ControlInputs = nullptr;
 		FInputInterface* StateInputs = nullptr;
 		bool bKeepVehicleAwake = false;
+		float CurrentTimeDilation = 1.0f;
 	};
 
 	struct FCreatedModule
@@ -598,6 +599,7 @@ namespace Chaos
 		virtual FSimOutputData* MakeNewData() = 0;
 		UE_API virtual void FillOutputState(const ISimulationModuleBase* SimModule);
 		UE_API virtual void Lerp(const FSimOutputData& InCurrent, const FSimOutputData& InNext, float Alpha);
+		UE_API virtual void GetFinalAnimDataGameThread(const FTransform& NewTransform, Chaos::FSimModuleAnimationData& AnimDataOut) { AnimDataOut = AnimationData; }
 
 		bool bEnabled = true;
 		FSimModuleAnimationData AnimationData;
diff --git a/Engine/Source/Runtime/Experimental/ChaosVehicles/ChaosVehiclesCore/Public/SimModule/SuspensionBaseInterface.h b/Engine/Source/Runtime/Experimental/ChaosVehicles/ChaosVehiclesCore/Public/SimModule/SuspensionBaseInterface.h
--- a/Engine/Source/Runtime/Experimental/ChaosVehicles/ChaosVehiclesCore/Public/SimModule/SuspensionBaseInterface.h
+++ b/Engine/Source/Runtime/Experimental/ChaosVehicles/ChaosVehiclesCore/Public/SimModule/SuspensionBaseInterface.h
@@ -69,9 +69,7 @@ public:
 	
 	virtual float GetSpringLength() const = 0;
 	virtual void SetSpringLength(float InLength, float WheelRadius) = 0;
-	
-	virtual void GetWorldRaycastLocation(const FTransform& BodyTransform, float WheelRadius, FSpringTrace& OutTrace) = 0;
-
+	UE_API virtual void GetWorldTraceEndpoints(float DeltaSeconds, const FTransform& BodyTransform, const FVector& Velocity, float WheelRadius, FSpringTrace& OutTrace) const = 0;
 	UE_API void SetTargetPoint(const FSuspensionTargetPoint& InTargetPoint);
 	const FSuspensionTargetPoint& GetTargetPoint() const { return TargetPoint; }
 	bool IsWheelInContact() const { return TargetPoint.bWheelInContact; }
