Modular vehicle : Can setup vehicle using only static meshes or skeletal meshes without needing cluster unions (or geometry collections). The existing setups still exist/work.

#Synced-CL 51943232

diff --git a/Engine/Plugins/Experimental/ChaosModularVehicle/Source/ChaosModularVehicleEngine/Private/ChaosModularVehicle/ModularVehicleBaseComponent.cpp b/Engine/Plugins/Experimental/ChaosModularVehicle/Source/ChaosModularVehicleEngine/Private/ChaosModularVehicle/ModularVehicleBaseComponent.cpp
--- a/Engine/Plugins/Experimental/ChaosModularVehicle/Source/ChaosModularVehicleEngine/Private/ChaosModularVehicle/ModularVehicleBaseComponent.cpp
+++ b/Engine/Plugins/Experimental/ChaosModularVehicle/Source/ChaosModularVehicleEngine/Private/ChaosModularVehicle/ModularVehicleBaseComponent.cpp
@@ -451,6 +451,56 @@ void UModularVehicleBaseComponent::CreateAssociatedSimComponents(USceneComponent
 	UpdatePhysicalProperties();
 }
 
+/**
+ * Helper function to find the visual component that should be animated.
+ * Walks up the component tree from the sim component to find the first
+ * StaticMeshComponent or other visual component.
+ * To be used at initialisation and not every frame for perf reasons.
+ *
+ * @param SimComponent - The sim component to start searching from
+ * @return The visual component to animate, or nullptr if not found
+ */
+static USceneComponent* FindVisualComponentToAnimate(USceneComponent* SimComponent)
+{
+	if (!SimComponent)
+	{
+		return nullptr;
+	}
+
+	USceneComponent* Current = SimComponent->GetAttachParent();
+
+	while (Current)
+	{
+		// Check if this is a visual component we want to animate
+
+		// StaticMeshComponent
+		if (UStaticMeshComponent* StaticMesh = Cast<UStaticMeshComponent>(Current))
+		{
+			return StaticMesh;
+		}
+
+		// SkeletalMeshComponent
+		if (USkeletalMeshComponent* SkeletalMesh = Cast<USkeletalMeshComponent>(Current))
+		{
+			// Animate skeletal meshes that are not the root component
+			return SkeletalMesh;
+		}
+
+		// Check if this is a generic scene component with a visual (child static mesh)      
+		// If we reach a scene component that has NO static/skeletal mesh children,
+		// we might want to animate it anyway (for hierarchical setups)
+
+		// Move up one level
+		Current = Current->GetAttachParent();
+	}
+
+	return nullptr;
+}
+
+/**
+* This code path is used only when we are not using the cluster union as a root component for part assembly
+* i.e. when Skeletal Mesh or Static Meshes are the only physics representation
+*/
 void UModularVehicleBaseComponent::CreateIndependentSimComponents(USceneComponent* RootComponent, USceneComponent* AttachedComponent, int ParentIndex, int TransformIndex, Chaos::FSimTreeUpdates& TreeUpdatesOut)
 {
 	using namespace Chaos;
@@ -461,10 +511,10 @@ void UModularVehicleBaseComponent::CreateIndependentSimComponents(USceneComponen
 
 	if (RootPhysicsObject == nullptr)
 	{
-		CacheRootPhysicsObject(GetPhysicsProxy()); //#TODO where to put / refresh?
+		CacheRootPhysicsObject(GetPhysicsProxy());
 	}
 
-	UE_LOG(LogModularBase, Log, TEXT("CreateAssociatedSimComponents Attaching %s, TransformIndex %d"), *AttachedComponent->GetName(), TransformIndex);
+	UE_LOG(LogModularBase, Log, TEXT("CreateIndependentSimComponents Attaching %s, TransformIndex %d"), *AttachedComponent->GetName(), TransformIndex);
 
 	if (IVehicleSimBaseComponentInterface* ComponentInterface = Cast<IVehicleSimBaseComponentInterface>(AttachedComponent))
 	{
@@ -473,13 +523,55 @@ void UModularVehicleBaseComponent::CreateIndependentSimComponents(USceneComponen
 			FTransform RootComponentTransform = RootComponent->GetComponentTransform();
 			FTransform ComponentTransform = AttachedComponent->GetComponentTransform().GetRelativeTransform(RootComponentTransform);
 
+			// Get the static or skeletal mesh parent of the simulation component
+			USceneComponent* VisualComponent = FindVisualComponentToAnimate(AttachedComponent);
+			if (VisualComponent && (VisualComponent != RootComponent))
+			{
+				ComponentTransform = VisualComponent->GetComponentTransform().GetRelativeTransform(RootComponentTransform);
+			}
+
+			const FName BoneName = ComponentInterface->GetBoneName();
+
+			// Try to get bone transform from RootComponent (if it's a skeletal mesh)
+			// This handles the original case where all modules are part of a single skeletal mesh
+			if (USkeletalMeshComponent* RootSKMComponent = Cast<USkeletalMeshComponent>(RootComponent))
+			{
+				if (BoneName != NAME_None)
+				{
+					for (int32 ArrayIndex = 0; ArrayIndex < RootSKMComponent->Bodies.Num(); ArrayIndex++)
+					{
+						if (FBodyInstance* BodyInstance = RootSKMComponent->Bodies[ArrayIndex])
+						{
+							if (BodyInstance->IsValidBodyInstance())
+							{
+								if (UBodySetup* BodySetup = BodyInstance->GetBodySetup())
+								{
+									const FName BodyName = BodySetup->BoneName;
+
+									if (BodyName == BoneName)
+									{
+										ComponentTransform = BodyInstance->GetUnrealWorldTransform().GetRelativeTransform(RootComponentTransform);
+
+										BodyInstance->bSimulatePhysics = false;
+										RootSKMComponent->bBlendPhysics = false;
+
+										UE_LOG(LogModularBase, Log, TEXT("  -> Found bone '%s' in ROOT skeletal mesh '%s'"), *BoneName.ToString(), *RootSKMComponent->GetName());
+										break;
+									}
+								}
+							}
+						}
+					}
+				}
+			}
+
 			int TreeIndex = INDEX_NONE;
 			TUniquePtr<Chaos::FSimModuleTree>& SimModuleTree = VehicleSimulationPT->AccessSimComponentTree();
 
 	#if !(UE_BUILD_SHIPPING || UE_BUILD_TEST)
 			FString DebugString;
 			NewModule->GetDebugString(DebugString);
-			UE_LOG(LogModularBase, Log, TEXT("CreateAssociatedSimComponents Found Sim Component %s (%s)"), *NewModule->GetDebugName(), *DebugString);
+			UE_LOG(LogModularBase, Log, TEXT("CreateIndependentSimComponents Found Sim Component %s (%s)"), *NewModule->GetDebugName(), *DebugString);
 	#endif
 
 			FVector LocalOffset(0, 0, 0);
@@ -501,6 +593,12 @@ void UModularVehicleBaseComponent::CreateIndependentSimComponents(USceneComponen
 			ComponentData.Guid = GenerateNewGuid();
 			NewModule->SetGuid(ComponentData.Guid);
 
+			// Where we want to animate the component transform directly and not using an AnimationInstance
+			if (VisualComponent != RootComponent)
+			{
+				ComponentData.VisualComponentToAnimate = VisualComponent;
+			}
+
 			PhysicsGuidToComponent.Add(ComponentData.Guid, AttachedComponent);
 
 			NewModule->SetInitialParticleTransform(InitialTransform);
@@ -518,12 +616,15 @@ void UModularVehicleBaseComponent::CreateIndependentSimComponents(USceneComponen
 				// if bone already exists then use that (separate wheel and suspension modules can share same bone)
 
 				int FoundIndex = INDEX_NONE;
-				for (int I = 0; I < ModuleAnimationSetups.Num(); I++)
+				if (ComponentBoneName != NAME_None)
 				{
-					if (ModuleAnimationSetups[I].BoneName == ComponentBoneName)
+					for (int I = 0; I < ModuleAnimationSetups.Num(); I++)
 					{
-						FoundIndex = I;
-						break;
+						if (ModuleAnimationSetups[I].BoneName == ComponentBoneName)
+						{
+							FoundIndex = I;
+							break;
+						}
 					}
 				}
 
@@ -536,6 +637,13 @@ void UModularVehicleBaseComponent::CreateIndependentSimComponents(USceneComponen
 				{
 					NewModule->SetAnimationData(ComponentBoneName, ComponentAnimationOffset, ModuleAnimationSetups.Num());
 					FModuleAnimationSetup AnimSetup(NewModule->GetBoneName(), TransformIndex);
+
+					// Store initial transform for animation
+					// This ensures that animated static meshes can calculate deltas from their initial pose.
+					AnimSetup.ModuleGUID = ComponentData.Guid;  // Link AnimSetup to module
+					AnimSetup.InitialRotOffset = ComponentTransform.GetRotation();
+					AnimSetup.InitialLocOffset = ComponentTransform.GetTranslation();
+
 					ModuleAnimationSetups.Add(AnimSetup);
 				}
 			}
@@ -558,7 +666,7 @@ void UModularVehicleBaseComponent::CreateIndependentSimComponents(USceneComponen
 	// recurse down tree, converting all SimComponents to proper simulation modules
 	for (USceneComponent* Child : Children)
 	{
-		if (IVehicleSimBaseComponentInterface* ChildSimComponent = Cast<IVehicleSimBaseComponentInterface>(Child))
+		if (USceneComponent* ChildSimComponent = Cast<USceneComponent>(Child))
 		{
 			CreateIndependentSimComponents(RootComponent, Child, ParentIndex, TransformIndex, TreeUpdatesOut);
 		}
@@ -667,10 +775,41 @@ void UModularVehicleBaseComponent::BeginPlay()
 		{
 			AddOverlappingComponentsToCluster();
 		}
-		else
+		else if (ClusterUnionComponent)
 		{
+			// add components to cluster union
 			AddComponentsFromOwnedActor();
 		}
+		else
+		{
+			if (VehicleSimulationPT)
+			{
+				if (AActor* Owner = GetOwner())
+				{
+					if (USceneComponent* RootBodyComponent = Owner->GetRootComponent())
+					{
+						if (UPrimitiveComponent* Prim = Cast<UPrimitiveComponent>(RootBodyComponent))
+						{
+							if (FBodyInstance* BI = Prim->GetBodyInstance())
+							{
+								// Get physics proxy from root body
+								IPhysicsProxyBase* Proxy = BI->GetPhysicsActorHandle();
+								CacheRootPhysicsObject(Proxy);
+							}
+
+							int ParentID = FindParentsLastSimComponent(RootBodyComponent);
+
+							Chaos::FSimTreeUpdates LatestTreeUpdates;
+							CreateIndependentSimComponents(RootBodyComponent, RootBodyComponent, ParentID, NextTransformIndex, LatestTreeUpdates);
+
+							Chaos::FSimTreeUpdates* NextTreeUpdates = nullptr;
+							NextTreeUpdates = &LatestTreeUpdates;
+							ActionTreeUpdates(NextTreeUpdates);
+						}
+					}
+				}
+			}
+		}
 	}
 
 }
@@ -1062,9 +1201,9 @@ void UModularVehicleBaseComponent::ParallelUpdate(const Chaos::FCreatedModules&
 				if (Chaos::FSimOutputData* ModuleOutput = PVehicleOutput->SimTreeOutputData[I])
 				{
 					FTransform NewTransform = FTransform::Identity;
-					if (ClusterUnionComponent)
+					if ( GetOwner() &&  GetOwner()->GetRootComponent())
 					{
-						NewTransform = ClusterUnionComponent->GetComponentToWorld();
+						NewTransform = GetOwner()->GetRootComponent()->GetComponentToWorld();
 					}
 
 					Chaos::FSimModuleAnimationData AnimData;
@@ -1097,6 +1236,7 @@ void UModularVehicleBaseComponent::ParallelUpdate(const Chaos::FCreatedModules&
 
 			}
 
+			UpdateNonSkeletalAnimations();
 
 		}
 	}
@@ -1119,21 +1259,21 @@ void UModularVehicleBaseComponent::PostUpdate()
 		{
 			BroadcastModuleAddedEvent(AddedModule.SimType, AddedModule.Guid, AddedModule.TreeIndex);
 		}
-	}
 
-	// Update replication structures
-	if (!PVehicleOutput->NewlyCreatedModuleGuids.IsEmpty())
-	{
-		if (NetworkPhysicsComponent)
+		// Update replication structures
+		if (!PVehicleOutput->NewlyCreatedModuleGuids.IsEmpty())
 		{
-			TSharedPtr<Chaos::FBaseRewindHistory>& History = NetworkPhysicsComponent->GetStateHistory_Internal();
-			Chaos::TDataRewindHistory<FNetworkModularVehicleStates>* StateHistory = static_cast<Chaos::TDataRewindHistory<FNetworkModularVehicleStates>*>(History.Get());
-			if (StateHistory)
+			if (NetworkPhysicsComponent)
 			{
-				for (int I = 0; I < StateHistory->GetDataHistory().Num(); I++)
+				TSharedPtr<Chaos::FBaseRewindHistory>& History = NetworkPhysicsComponent->GetStateHistory_Internal();
+				Chaos::TDataRewindHistory<FNetworkModularVehicleStates>* StateHistory = static_cast<Chaos::TDataRewindHistory<FNetworkModularVehicleStates>*>(History.Get());
+				if (StateHistory)
 				{
-					FNetworkModularVehicleStates& State = StateHistory->GetDataHistory()[I];
-					VehicleSimulationPT->GenerateReplicationStructure(State);
+					for (int I = 0; I < StateHistory->GetDataHistory().Num(); I++)
+					{
+						FNetworkModularVehicleStates& State = StateHistory->GetDataHistory()[I];
+						VehicleSimulationPT->GenerateReplicationStructure(State);
+					}
 				}
 			}
 		}
@@ -1151,6 +1291,61 @@ const FTransform& UModularVehicleBaseComponent::GetComponentTransform() const
 	return MyComponent->GetComponentTransform();
 }
 
+/**
+ * Updates transforms for non-skeletal components based on animation data.
+ * This allows static meshes and scene components to be animated by the physics simulation.
+ *
+ * Call this from ParallelUpdate() after ModuleAnimationSetups is populated with physics data.
+ */
+void UModularVehicleBaseComponent::UpdateNonSkeletalAnimations()
+{
+	for (int32 AnimIndex = 0; AnimIndex < ModuleAnimationSetups.Num(); AnimIndex++)
+	{
+		const FModuleAnimationSetup& AnimSetup = ModuleAnimationSetups[AnimIndex];
+		int32 ModuleGuid = AnimSetup.ModuleGUID;
+
+		if (ModuleGuid == INDEX_NONE)
+		{
+			continue;
+		}
+
+		// Find the component data
+		for (TPair<TObjectPtr<USceneComponent>, FVehicleComponentData>& Pair : ComponentToPhysicsObjects)
+		{
+			if (Pair.Value.Guid == ModuleGuid)
+			{
+				// Use cached visual component
+				USceneComponent* ComponentToAnimate = Pair.Value.VisualComponentToAnimate;
+
+				if (!ComponentToAnimate || !ComponentToAnimate->IsValidLowLevel())
+				{
+					continue;
+				}
+
+				// Apply animation
+				FTransform NewRelativeTransform = ComponentToAnimate->GetRelativeTransform();
+
+				if (AnimSetup.AnimFlags & Chaos::EAnimationFlags::AnimateRotation)
+				{
+					NewRelativeTransform.SetRotation(AnimSetup.InitialRotOffset *
+						AnimSetup.CombinedRotation);
+				}
+
+				if (AnimSetup.AnimFlags & Chaos::EAnimationFlags::AnimatePosition)
+				{
+					NewRelativeTransform.SetLocation(AnimSetup.InitialLocOffset +
+						AnimSetup.LocOffset);
+				}
+
+				ComponentToAnimate->SetRelativeTransform(NewRelativeTransform, false, nullptr,
+					ETeleportType::TeleportPhysics);
+				break;
+			}
+		}
+	}
+}
+
+
 void UModularVehicleBaseComponent::ActionTreeUpdates(Chaos::FSimTreeUpdates* NextTreeUpdates)
 {
 	if (VehicleSimulationPT)
@@ -1431,41 +1626,42 @@ void UModularVehicleBaseComponent::CreateVehicleSim()
 
 void UModularVehicleBaseComponent::DestroyVehicleSim()
 {
-	UWorld* World = GetWorld();
-	if (World->IsGameWorld())
+	if (UWorld* World = GetWorld())
 	{
-		UE_LOG(LogModularBase, Log, TEXT("UModularVehicleBaseComponent::DestroyVehicleSim"));
+		if (World->IsGameWorld())
+		{
+			UE_LOG(LogModularBase, Log, TEXT("UModularVehicleBaseComponent::DestroyVehicleSim"));
 
-		FPhysScene* PhysScene = World->GetPhysicsScene();
-		check(PhysScene);
+			FPhysScene* PhysScene = World->GetPhysicsScene();
+			check(PhysScene);
 
-		if (Chaos::FPhysicsSolver* LocalSolver = PhysScene->GetSolver())
-		{
-			if (FChaosSimModuleManager* SimManager = FChaosSimModuleManager::GetManagerFromScene(PhysScene))
+			if (Chaos::FPhysicsSolver* LocalSolver = PhysScene->GetSolver())
 			{
-				SimManager->RemoveVehicle(this);
+				if (FChaosSimModuleManager* SimManager = FChaosSimModuleManager::GetManagerFromScene(PhysScene))
+				{
+					SimManager->RemoveVehicle(this);
+				}
+
+				if (PVehicleOutput.IsValid())
+				{
+					PVehicleOutput.Reset(nullptr);
+				}
+
+				// postpone deletion of object while simultaneously setting it to null in the component (TUniquePtr)
+				LocalSolver->EnqueueCommandImmediate([Pointer = MoveTemp(VehicleSimulationPT)]() mutable
+					{
+						Pointer.Reset(nullptr);
+					});
+
 			}
 
-			if (PVehicleOutput.IsValid())
+			if (bUsingNetworkPhysicsPrediction && NetworkPhysicsComponent)
 			{
-				PVehicleOutput.Reset(nullptr);
+				NetworkPhysicsComponent->RemoveDataHistory();
 			}
 
-			// postpone deletion of object while simultaneously setting it to null in the component (TUniquePtr)
-			LocalSolver->EnqueueCommandImmediate([Pointer = MoveTemp(VehicleSimulationPT)]() mutable
-				{
-					Pointer.Reset(nullptr);
-				});
-
 		}
-
-		if (bUsingNetworkPhysicsPrediction && NetworkPhysicsComponent)
-		{
-			NetworkPhysicsComponent->RemoveDataHistory();
-		}
-
 	}
-
 }
 
 
diff --git a/Engine/Plugins/Experimental/ChaosModularVehicle/Source/ChaosModularVehicleEngine/Private/ChaosModularVehicle/ModularVehicleSkeletalPawn.cpp b/Engine/Plugins/Experimental/ChaosModularVehicle/Source/ChaosModularVehicleEngine/Private/ChaosModularVehicle/ModularVehicleSkeletalPawn.cpp
new file mode 100644
--- /dev/null
+++ b/Engine/Plugins/Experimental/ChaosModularVehicle/Source/ChaosModularVehicleEngine/Private/ChaosModularVehicle/ModularVehicleSkeletalPawn.cpp
@@ -0,0 +1,22 @@
+// Copyright Epic Games, Inc. All Rights Reserved.
+
+#include "ChaosModularVehicle/ModularVehicleSkeletalPawn.h"
+#include "Components/SkeletalMeshComponent.h"
+#include "ChaosModularVehicle/ModularVehicleBaseComponent.h"
+
+#include UE_INLINE_GENERATED_CPP_BY_NAME(ModularVehicleSkeletalPawn)
+
+
+AModularVehicleSkeletalPawn::AModularVehicleSkeletalPawn(const FObjectInitializer& ObjectInitializer)
+	: Super(ObjectInitializer)
+{
+	SkeletalMeshComponent = CreateDefaultSubobject<USkeletalMeshComponent>(TEXT("SkeletalMeshComponent0"));
+	SetRootComponent(SkeletalMeshComponent);
+
+	VehicleSimComponent = CreateDefaultSubobject<UModularVehicleBaseComponent>(TEXT("VehicleSimComponent0"));
+	VehicleSimComponent->SetClusterComponent(nullptr); // no cluster union component in this scenario
+
+	SetRemoteRoleForBackwardsCompat(ROLE_SimulatedProxy);
+	bReplicates = true;
+	SetReplicatingMovement(true);
+}
diff --git a/Engine/Plugins/Experimental/ChaosModularVehicle/Source/ChaosModularVehicleEngine/Public/ChaosModularVehicle/ModularVehicleBaseComponent.h b/Engine/Plugins/Experimental/ChaosModularVehicle/Source/ChaosModularVehicleEngine/Public/ChaosModularVehicle/ModularVehicleBaseComponent.h
--- a/Engine/Plugins/Experimental/ChaosModularVehicle/Source/ChaosModularVehicleEngine/Public/ChaosModularVehicle/ModularVehicleBaseComponent.h
+++ b/Engine/Plugins/Experimental/ChaosModularVehicle/Source/ChaosModularVehicleEngine/Public/ChaosModularVehicle/ModularVehicleBaseComponent.h
@@ -34,6 +34,11 @@ struct FVehicleComponentData
 	GENERATED_BODY()
 
 	int Guid = -1;
+
+	// Cached reference to the visual component to animate (if different from sim component)
+	UPROPERTY()
+	TObjectPtr<USceneComponent> VisualComponentToAnimate = nullptr;
+
 };
 
 /** Additional replicated state */
@@ -307,7 +312,8 @@ public:
 	int32 DelayClusteringCount;
 
 	/*** Map simulation component to our vehicle setup data */
-	TMap<TObjectKey<USceneComponent>, FVehicleComponentData> ComponentToPhysicsObjects;
+	UPROPERTY()
+	TMap<TObjectPtr<USceneComponent>, FVehicleComponentData> ComponentToPhysicsObjects;
 
 	/*** Map physics sim GUID to actor component */
 	TMap<int, TWeakObjectPtr<USceneComponent>> PhysicsGuidToComponent;
@@ -349,6 +355,7 @@ protected:
 	UE_API void AddComponentToCluster(USceneComponent* Component);
 	UE_API bool AddComponentsFromOwnedActor();
 	UE_API void SetupSkeletalAnimationStructure();
+	UE_API void UpdateNonSkeletalAnimations();
 	UE_API void AssimilateComponentInputs(TArray<FModuleInputSetup>& OutCombinedInputs);
 	UE_API void CacheRootPhysicsObject(IPhysicsProxyBase* Proxy);
 
diff --git a/Engine/Plugins/Experimental/ChaosModularVehicle/Source/ChaosModularVehicleEngine/Public/ChaosModularVehicle/ModularVehicleSkeletalPawn.h b/Engine/Plugins/Experimental/ChaosModularVehicle/Source/ChaosModularVehicleEngine/Public/ChaosModularVehicle/ModularVehicleSkeletalPawn.h
new file mode 100644
--- /dev/null
+++ b/Engine/Plugins/Experimental/ChaosModularVehicle/Source/ChaosModularVehicleEngine/Public/ChaosModularVehicle/ModularVehicleSkeletalPawn.h
@@ -0,0 +1,33 @@
+// Copyright Epic Games, Inc. All Rights Reserved.
+
+#pragma once
+
+#include "CoreMinimal.h"
+#include "GameFramework/Pawn.h"
+
+#include "ModularVehicleSkeletalPawn.generated.h"
+
+class USkeletalMeshComponent;
+class UModularVehicleBaseComponent;
+
+UCLASS(MinimalAPI)
+class AModularVehicleSkeletalPawn: public APawn
+{
+	GENERATED_UCLASS_BODY()
+
+public:
+	UFUNCTION()
+	USkeletalMeshComponent* GetSkeletalMeshComponent() const { return SkeletalMeshComponent; }
+
+	UFUNCTION()
+	UModularVehicleBaseComponent* GetVehicleSimulationComponent() const { return VehicleSimComponent; }
+
+	/* SkeletalMeshComponent root component */
+	UPROPERTY(VisibleAnywhere, BlueprintReadOnly, Category = Vehicle, meta = (ExposeFunctionCategories = "Components|ModularVehicle", AllowPrivateAccess = "true"))
+	TObjectPtr<USkeletalMeshComponent> SkeletalMeshComponent;
+
+	/* ModularVehicleComponent */
+	UPROPERTY(VisibleAnywhere, BlueprintReadOnly, Category = Vehicle, meta = (ExposeFunctionCategories = "Components|ModularVehicle", AllowPrivateAccess = "true"))
+	TObjectPtr<UModularVehicleBaseComponent> VehicleSimComponent;
+
+};
