[ModularVehicle] Rely on NetworkPhysicsComponent.IsLocallyControlled from the Modular Vehicle instead of having two overlapping solutions for overriding local control.

#Synced-CL 52463959

diff --git a/Engine/Plugins/Experimental/ChaosModularVehicle/Source/ChaosModularVehicleEngine/Private/ChaosModularVehicle/ModularVehicleBaseComponent.cpp b/Engine/Plugins/Experimental/ChaosModularVehicle/Source/ChaosModularVehicleEngine/Private/ChaosModularVehicle/ModularVehicleBaseComponent.cpp
--- a/Engine/Plugins/Experimental/ChaosModularVehicle/Source/ChaosModularVehicleEngine/Private/ChaosModularVehicle/ModularVehicleBaseComponent.cpp
+++ b/Engine/Plugins/Experimental/ChaosModularVehicle/Source/ChaosModularVehicleEngine/Private/ChaosModularVehicle/ModularVehicleBaseComponent.cpp
@@ -108,6 +108,15 @@ APlayerController* UModularVehicleBaseComponent::GetPlayerController() const
 
 bool UModularVehicleBaseComponent::IsLocallyControlled() const
 {
+	// When using network physics prediction, the NetworkPhysicsComponent is the single source of truth
+	// for local control state. It tracks both the force-override (bIsRelayingLocalInputs)
+	// and the controller-based check.
+	if (bUsingNetworkPhysicsPrediction && NetworkPhysicsComponent)
+	{
+		return NetworkPhysicsComponent->IsLocallyControlled();
+	}
+
+	// Fallback for when NetworkPhysicsComponent is not available
 	if (bIsLocallyControlled)
 	{
 		return true;
@@ -200,7 +209,7 @@ void UModularVehicleBaseComponent::OnCreatePhysicsState()
 
 			if (bIsLocallyControlled)
 			{
-				NetworkPhysicsComponent->SetIsRelayingLocalInputs(bIsLocallyControlled);
+				NetworkPhysicsComponent->SetIsRelayingLocalInputs(true);
 			}
 
 		}
@@ -1545,21 +1554,29 @@ void UModularVehicleBaseComponent::RemoveComponentFromSimulationImpl(UPrimitiveC
 
 void UModularVehicleBaseComponent::SetLocallyControlled(bool bLocallyControlledIn)
 {
-	bIsLocallyControlled = false;
-	if (UWorld* World = GetWorld())
+	if (bUsingNetworkPhysicsPrediction && NetworkPhysicsComponent)
 	{
-		// guard against invalid case that can lead to bad networking state
+		// NetworkPhysicsComponent is the single source of truth, guard against sim-proxy misuse
 		if (GetOwner() && GetOwner()->GetLocalRole() != ENetRole::ROLE_SimulatedProxy)
 		{
-			bIsLocallyControlled = bLocallyControlledIn;
+			NetworkPhysicsComponent->SetIsRelayingLocalInputs(bLocallyControlledIn);
+		}
+		else
+		{
+			NetworkPhysicsComponent->SetIsRelayingLocalInputs(false);
 		}
+		return;
 	}
 
-	if (bUsingNetworkPhysicsPrediction && NetworkPhysicsComponent)
+	// Fallback for when NetworkPhysicsComponent is not available
+	bIsLocallyControlled = false;
+	if (UWorld* World = GetWorld())
 	{
-		NetworkPhysicsComponent->SetIsRelayingLocalInputs(bIsLocallyControlled);
+		if (GetOwner() && GetOwner()->GetLocalRole() != ENetRole::ROLE_SimulatedProxy)
+		{
+			bIsLocallyControlled = bLocallyControlledIn;
+		}
 	}
-
 }
 
 void UModularVehicleBaseComponent::AssimilateComponentInputs(TArray<FModuleInputSetup>& OutCombinedInputs)
