[NetPhysics] Fix IsLocallyControlled ensure on physics thread in ModularVehicle

#Synced-CL 52778023

diff --git a/Engine/Plugins/Experimental/ChaosModularVehicle/Source/ChaosModularVehicleEngine/Private/ChaosModularVehicle/ChaosSimModuleManagerAsyncCallback.cpp b/Engine/Plugins/Experimental/ChaosModularVehicle/Source/ChaosModularVehicleEngine/Private/ChaosModularVehicle/ChaosSimModuleManagerAsyncCallback.cpp
--- a/Engine/Plugins/Experimental/ChaosModularVehicle/Source/ChaosModularVehicleEngine/Private/ChaosModularVehicle/ChaosSimModuleManagerAsyncCallback.cpp
+++ b/Engine/Plugins/Experimental/ChaosModularVehicle/Source/ChaosModularVehicleEngine/Private/ChaosModularVehicle/ChaosSimModuleManagerAsyncCallback.cpp
@@ -268,7 +268,7 @@ void FModularVehicleAsyncInput::ProcessInputs()
 		}
 	}
 
-	if (GetVehicle()->IsLocallyControlled() && !bIsResimming)
+	if (bIsLocallyControlled && !bIsResimming)
 	{
 		VehicleSim->VehicleInputs = PhysicsInputs.NetworkInputs.VehicleInputs;
 	}
diff --git a/Engine/Plugins/Experimental/ChaosModularVehicle/Source/ChaosModularVehicleEngine/Private/ChaosModularVehicle/ModularVehicleBaseComponent.cpp b/Engine/Plugins/Experimental/ChaosModularVehicle/Source/ChaosModularVehicleEngine/Private/ChaosModularVehicle/ModularVehicleBaseComponent.cpp
--- a/Engine/Plugins/Experimental/ChaosModularVehicle/Source/ChaosModularVehicleEngine/Private/ChaosModularVehicle/ModularVehicleBaseComponent.cpp
+++ b/Engine/Plugins/Experimental/ChaosModularVehicle/Source/ChaosModularVehicleEngine/Private/ChaosModularVehicle/ModularVehicleBaseComponent.cpp
@@ -139,6 +139,7 @@ void UModularVehicleBaseComponent::ProduceInput(int32 PhysicsStep, int32 NumStep
 	if (AsyncInput && GetPhysicsProxy())
 	{
 		AsyncInput->SetVehicle(this);
+		AsyncInput->bIsLocallyControlled = IsLocallyControlled();
 		CurAsyncInput = AsyncInput;
 		AsyncInput->Proxy = GetPhysicsProxy();
 		AsyncInput->PhysicsInputs.NetworkInputs.VehicleInputs.KeepAwake = bKeepVehicleAwake;
diff --git a/Engine/Plugins/Experimental/ChaosModularVehicle/Source/ChaosModularVehicleEngine/Public/ChaosModularVehicle/ChaosSimModuleManagerAsyncCallback.h b/Engine/Plugins/Experimental/ChaosModularVehicle/Source/ChaosModularVehicleEngine/Public/ChaosModularVehicle/ChaosSimModuleManagerAsyncCallback.h
--- a/Engine/Plugins/Experimental/ChaosModularVehicle/Source/ChaosModularVehicleEngine/Public/ChaosModularVehicle/ChaosSimModuleManagerAsyncCallback.h
+++ b/Engine/Plugins/Experimental/ChaosModularVehicle/Source/ChaosModularVehicleEngine/Public/ChaosModularVehicle/ChaosSimModuleManagerAsyncCallback.h
@@ -280,6 +280,7 @@ struct FModularVehicleAsyncInput
 
 	const EChaosAsyncVehicleDataType Type;
 	IPhysicsProxyBase* Proxy;
+	bool bIsLocallyControlled = false;
 
 	FPhysicsModularVehicleInputs PhysicsInputs;
 
